Datatable editor now refreshes correctly when the datatable is reimported

diff --git a/Engine/Source/Editor/DataTableEditor/Private/DataTableEditor.cpp b/Engine/Source/Editor/DataTableEditor/Private/DataTableEditor.cpp
--- a/Engine/Source/Editor/DataTableEditor/Private/DataTableEditor.cpp
+++ b/Engine/Source/Editor/DataTableEditor/Private/DataTableEditor.cpp
@@ -24,9 +24,14 @@ void FDataTableEditor::UnregisterTabSpawners(const TSharedRef<class FTabManager>
 	TabManager->UnregisterTabSpawner( DataTableTabId );
 }
 
-FDataTableEditor::~FDataTableEditor()
+FDataTableEditor::FDataTableEditor()
 {
+	FReimportManager::Instance()->OnPostReimport().AddRaw(this, &FDataTableEditor::OnPostReimport);
+}
 
+FDataTableEditor::~FDataTableEditor()
+{
+	FReimportManager::Instance()->OnPostReimport().RemoveAll(this);
 }
 
 
@@ -83,6 +88,16 @@ FLinearColor FDataTableEditor::GetWorldCentricTabColorScale() const
 	return FLinearColor( 0.0f, 0.0f, 0.2f, 0.5f );
 }
 
+void FDataTableEditor::OnPostReimport(UObject* InObject, bool bSuccess)
+{
+	// Ignore if this is regarding a different object
+	if ( InObject == DataTable && bSuccess)
+	{
+		CachedDataTable.Empty();
+		ReloadVisibleData();
+	}
+}
+
 TSharedPtr<SUniformGridPanel> FDataTableEditor::CreateGridPanel()
 {
 	TSharedPtr<SUniformGridPanel> GridPanel = SNew(SUniformGridPanel).SlotPadding( FMargin( 2.0f ) );
diff --git a/Engine/Source/Editor/DataTableEditor/Private/DataTableEditor.h b/Engine/Source/Editor/DataTableEditor/Private/DataTableEditor.h
--- a/Engine/Source/Editor/DataTableEditor/Private/DataTableEditor.h
+++ b/Engine/Source/Editor/DataTableEditor/Private/DataTableEditor.h
@@ -25,6 +25,9 @@ public:
 	 */
 	void InitDataTableEditor( const EToolkitMode::Type Mode, const TSharedPtr< class IToolkitHost >& InitToolkitHost, UDataTable* Table );
 
+	/** Constructor */
+	FDataTableEditor();
+
 	/** Destructor */
 	virtual ~FDataTableEditor();
 
@@ -43,6 +46,9 @@ public:
 
 private:
 
+	/** Called when an object is reimported */
+	void OnPostReimport(UObject* InObject, bool bSuccess);
+
 	void ReloadVisibleData();
 
 	void OnSearchTextChanged(const FText& SearchText);
