Data Table Row Editor first version.
Editor Utils refactored.
Some fixes in GetDataTableRow.

diff --git a/Engine/Source/Editor/BlueprintGraph/Classes/K2Node_GetDataTableRow.h b/Engine/Source/Editor/BlueprintGraph/Classes/K2Node_GetDataTableRow.h
--- a/Engine/Source/Editor/BlueprintGraph/Classes/K2Node_GetDataTableRow.h
+++ b/Engine/Source/Editor/BlueprintGraph/Classes/K2Node_GetDataTableRow.h
@@ -23,6 +23,7 @@ class BLUEPRINTGRAPH_API UK2Node_GetDataTableRow : public UK2Node
         OutColor = GetNodeTitleColor();
         return TEXT("Kismet.AllClasses.FunctionIcon");
     }
+	virtual void ValidateNodeDuringCompilation(class FCompilerResultsLog& MessageLog) const override;
 	// End UEdGraphNode interface.
 
 	// Begin UK2Node interface
@@ -53,9 +54,7 @@ class BLUEPRINTGRAPH_API UK2Node_GetDataTableRow : public UK2Node
 	/** Get the type of the TableRow to return */
 	UScriptStruct* GetDataTableRowStructType(const TArray<UEdGraphPin*>* InPinsToSearch = NULL) const;
 
-	/** Get the list of possible RowNames */
-	void GetDataTableRowNameList(TArray<TSharedPtr<FName>>& OutList, const TArray<UEdGraphPin*>* InPinsToSearch = NULL) const;
-
+	void OnDataTableRowListChanged(const UDataTable* DataTable);
 private:
 	/**
 	 * Takes the specified "MutatablePin" and sets its 'PinToolTip' field (according
diff --git a/Engine/Source/Editor/BlueprintGraph/Private/BlueprintGraphPanelPinFactory.cpp b/Engine/Source/Editor/BlueprintGraph/Private/BlueprintGraphPanelPinFactory.cpp
--- a/Engine/Source/Editor/BlueprintGraph/Private/BlueprintGraphPanelPinFactory.cpp
+++ b/Engine/Source/Editor/BlueprintGraph/Private/BlueprintGraphPanelPinFactory.cpp
@@ -3,6 +3,7 @@
 #include "BlueprintGraphPanelPinFactory.h"
 #include "K2Node_GetDataTableRow.h"
 #include "SGraphPinNameList.h"
+#include "SGraphPinDataTableRowName.h"
 #include "K2Node_CallFunction.h"
 	
 TSharedPtr<class SGraphPin> FBlueprintGraphPanelPinFactory::CreatePin(class UEdGraphPin* InPin) const
@@ -37,21 +38,9 @@ TSharedPtr<class SGraphPin> FBlueprintGraphPanelPinFactory::CreatePin(class UEdG
 		{
 			if (DataTablePin && DataTablePin->DefaultObject != nullptr && DataTablePin->LinkedTo.Num() == 0)
 			{
-				if (DataTablePin->DefaultObject->IsA(UDataTable::StaticClass()))
+				if (auto DataTable = Cast<UDataTable>(DataTablePin->DefaultObject))
 				{
-					UDataTable* DataTable = (UDataTable*)DataTablePin->DefaultObject;
-					if (DataTable)
-					{
-						TArray<TSharedPtr<FName>> RowNames;
-						/** Extract all the row names from the RowMap */
-						for (TMap<FName, uint8*>::TConstIterator Iterator(DataTable->RowMap); Iterator; ++Iterator)
-						{
-							/** Create a simple array of the row names */
-							TSharedPtr<FName> RowNameItem = MakeShareable(new FName(Iterator.Key()));
-							RowNames.Add(RowNameItem);
-						}
-						return SNew(SGraphPinNameList, InPin, RowNames);
-					}
+					return SNew(SGraphPinDataTableRowName, InPin, DataTable);
 				}
 				if (DataTablePin->DefaultObject->IsA(UCurveTable::StaticClass()))
 				{
diff --git a/Engine/Source/Editor/BlueprintGraph/Private/K2Node_GetDataTableRow.cpp b/Engine/Source/Editor/BlueprintGraph/Private/K2Node_GetDataTableRow.cpp
--- a/Engine/Source/Editor/BlueprintGraph/Private/K2Node_GetDataTableRow.cpp
+++ b/Engine/Source/Editor/BlueprintGraph/Private/K2Node_GetDataTableRow.cpp
@@ -132,31 +132,24 @@ UScriptStruct* UK2Node_GetDataTableRow::GetDataTableRowStructType(const TArray<U
 	return RowStructType;
 }
 
-void UK2Node_GetDataTableRow::GetDataTableRowNameList(TArray<TSharedPtr<FName>>& RowNames, const TArray<UEdGraphPin*>* InPinsToSearch /*= NULL*/) const
+void UK2Node_GetDataTableRow::OnDataTableRowListChanged(const UDataTable* DataTable)
 {
-	const TArray<UEdGraphPin*>* PinsToSearch = InPinsToSearch ? InPinsToSearch : &Pins;
-
-	UEdGraphPin* DataTablePin = GetDataTablePin(PinsToSearch);
-	if (DataTablePin && DataTablePin->DefaultObject != NULL && DataTablePin->LinkedTo.Num() == 0)
+	UEdGraphPin* DataTablePin = GetDataTablePin();
+	if (DataTable && DataTablePin && DataTable == DataTablePin->DefaultObject)
 	{
-		if (DataTablePin->DefaultObject->IsA(UDataTable::StaticClass()))
+		auto RowNamePin = GetRowNamePin();
+		const bool TryRefresh = RowNamePin && !RowNamePin->LinkedTo.Num();
+		const FName CurrentName = RowNamePin ? FName(*RowNamePin->GetDefaultAsString()) : NAME_None;
+		if (TryRefresh && RowNamePin && !DataTable->GetRowNames().Contains(CurrentName))
 		{
-			UDataTable* DataTable = (UDataTable*)DataTablePin->DefaultObject;
-			if (DataTable)
+			if (auto BP = GetBlueprint())
 			{
-				/** Extract all the row names from the RowMap */
-				for (TMap<FName, uint8*>::TConstIterator Iterator(DataTable->RowMap); Iterator; ++Iterator)
-				{
-					/** Create a simple array of the row names */
-					TSharedPtr<FName> RowNameItem = MakeShareable(new FName(Iterator.Key()));
-					RowNames.Add(RowNameItem);
-				}
+				FBlueprintEditorUtils::MarkBlueprintAsModified(BP);
 			}
 		}
 	}
 }
 
-
 void UK2Node_GetDataTableRow::ReallocatePinsDuringReconstruction(TArray<UEdGraphPin*>& OldPins) 
 {
 	Super::ReallocatePinsDuringReconstruction(OldPins);
@@ -383,4 +376,33 @@ void UK2Node_GetDataTableRow::ExpandNode(class FKismetCompilerContext& CompilerC
 	BreakAllNodeLinks();
 }
 
+void UK2Node_GetDataTableRow::ValidateNodeDuringCompilation(class FCompilerResultsLog& MessageLog) const
+{
+	const auto DataTablePin = GetDataTablePin();
+	const auto RowNamePin = GetRowNamePin();
+	if (!DataTablePin || !RowNamePin)
+	{
+		MessageLog.Error(*LOCTEXT("MissingPins", "Missing pins in @@").ToString(), this);
+		return;
+	}
+
+	const auto DataTable = Cast<UDataTable>(DataTablePin->DefaultObject);
+	if (!DataTable)
+	{
+		MessageLog.Error(*LOCTEXT("NoDataTable", "No DataTable in @@").ToString(), this);
+		return;
+	}
+
+	const FName CurrentName = FName(*RowNamePin->GetDefaultAsString());
+	if (!DataTable->GetRowNames().Contains(CurrentName))
+	{
+		const FString Msg = FString::Printf(
+			*LOCTEXT("WronRowName", "'%s' row name is not stored in '%s'. @@").ToString()
+			, *CurrentName.ToString()
+			, *GetFullNameSafe(DataTable));
+		MessageLog.Error(*Msg, this);
+		return;
+	}
+}
+
 #undef LOCTEXT_NAMESPACE
\ No newline at end of file
diff --git a/Engine/Source/Editor/DataTableEditor/DataTableEditor.Build.cs b/Engine/Source/Editor/DataTableEditor/DataTableEditor.Build.cs
--- a/Engine/Source/Editor/DataTableEditor/DataTableEditor.Build.cs
+++ b/Engine/Source/Editor/DataTableEditor/DataTableEditor.Build.cs
@@ -13,13 +13,15 @@ public class DataTableEditor : ModuleRules
 				"Core",
 				"CoreUObject",
 				"Engine",
+                "InputCore",
 				"Slate",
 				"SlateCore",
                 "EditorStyle",
+                "PropertyEditor",
 				"UnrealEd"
 			}
 			);
 
 		DynamicallyLoadedModuleNames.Add("WorkspaceMenuStructure");
 	}
-}
+}
\ No newline at end of file
diff --git a/Engine/Source/Editor/DataTableEditor/Private/DataTableEditor.cpp b/Engine/Source/Editor/DataTableEditor/Private/DataTableEditor.cpp
--- a/Engine/Source/Editor/DataTableEditor/Private/DataTableEditor.cpp
+++ b/Engine/Source/Editor/DataTableEditor/Private/DataTableEditor.cpp
@@ -7,11 +7,13 @@
 #include "Editor/WorkspaceMenuStructure/Public/WorkspaceMenuStructureModule.h"
 #include "SSearchBox.h"
 #include "SDockTab.h"
+#include "SRowEditor.h"
 #include "Engine/DataTable.h"
  
 #define LOCTEXT_NAMESPACE "DataTableEditor"
 
 const FName FDataTableEditor::DataTableTabId( TEXT( "DataTableEditor_DataTable" ) );
+const FName FDataTableEditor::RowEditorTabId(TEXT("DataTableEditor_RowEditor"));
 
 void FDataTableEditor::RegisterTabSpawners(const TSharedRef<class FTabManager>& TabManager)
 {
@@ -20,35 +22,69 @@ void FDataTableEditor::RegisterTabSpawners(const TSharedRef<class FTabManager>&
 	TabManager->RegisterTabSpawner( DataTableTabId, FOnSpawnTab::CreateSP(this, &FDataTableEditor::SpawnTab_DataTable) )
 		.SetDisplayName( LOCTEXT("DataTableTab", "Data Table") )
 		.SetGroup( WorkspaceMenuCategory.ToSharedRef() );
+
+	TabManager->RegisterTabSpawner(RowEditorTabId, FOnSpawnTab::CreateSP(this, &FDataTableEditor::SpawnTab_RowEditor))
+		.SetDisplayName(LOCTEXT("RowEditorTab", "Row Editor"))
+		.SetGroup(WorkspaceMenuCategory.ToSharedRef());
 }
 
 void FDataTableEditor::UnregisterTabSpawners(const TSharedRef<class FTabManager>& TabManager)
 {
 	TabManager->UnregisterTabSpawner( DataTableTabId );
+	TabManager->UnregisterTabSpawner(RowEditorTabId);
 }
 
 FDataTableEditor::FDataTableEditor()
 {
-	FReimportManager::Instance()->OnPostReimport().AddRaw(this, &FDataTableEditor::OnPostReimport);
 }
 
 FDataTableEditor::~FDataTableEditor()
 {
-	FReimportManager::Instance()->OnPostReimport().RemoveAll(this);
 }
 
+void FDataTableEditor::PreChange(const class UUserDefinedStruct* Struct, FStructureEditorUtils::EStructureEditorChangeInfo Info)
+{
+}
+
+void FDataTableEditor::PostChange(const class UUserDefinedStruct* Struct, FStructureEditorUtils::EStructureEditorChangeInfo Info)
+{
+	if (Struct && DataTable.IsValid() && (DataTable->RowStruct == Struct))
+	{
+		CachedDataTable.Empty();
+		ReloadVisibleData();
+	}
+}
+
+void FDataTableEditor::PreChange(const UDataTable* Changed, FDataTableEditorUtils::EDataTableChangeInfo Info)
+{
+}
+
+void FDataTableEditor::PostChange(const UDataTable* Changed, FDataTableEditorUtils::EDataTableChangeInfo Info)
+{
+	FStringAssetReference::InvalidateTag(); // Should be removed after UE-5615 is fixed
+	if (Changed == DataTable.Get())
+	{
+		CachedDataTable.Empty();
+		ReloadVisibleData();
+	}
+}
 
 void FDataTableEditor::InitDataTableEditor( const EToolkitMode::Type Mode, const TSharedPtr< class IToolkitHost >& InitToolkitHost, UDataTable* Table )
 {
 	TSharedRef<FTabManager::FLayout> StandaloneDefaultLayout = FTabManager::NewLayout( "Standalone_DataTableEditor_Layout" )
 	->AddArea
 	(
-		FTabManager::NewPrimaryArea()
+		FTabManager::NewPrimaryArea()->SetOrientation(Orient_Vertical)
 		->Split
 		(
 			FTabManager::NewStack()
 			->AddTab( DataTableTabId, ETabState::OpenedTab )
 		)
+		->Split
+		(
+			FTabManager::NewStack()
+			->AddTab(RowEditorTabId, ETabState::OpenedTab)
+		)
 	);
 
 	const bool bCreateDefaultStandaloneMenu = true;
@@ -91,21 +127,11 @@ FLinearColor FDataTableEditor::GetWorldCentricTabColorScale() const
 	return FLinearColor( 0.0f, 0.0f, 0.2f, 0.5f );
 }
 
-void FDataTableEditor::OnPostReimport(UObject* InObject, bool bSuccess)
-{
-	// Ignore if this is regarding a different object
-	if ( InObject == DataTable && bSuccess)
-	{
-		CachedDataTable.Empty();
-		ReloadVisibleData();
-	}
-}
-
 TSharedPtr<SUniformGridPanel> FDataTableEditor::CreateGridPanel()
 {
 	TSharedPtr<SUniformGridPanel> GridPanel = SNew(SUniformGridPanel).SlotPadding( FMargin( 2.0f ) );
 
-	if (DataTable != NULL)
+	if (DataTable.IsValid())
 	{
 		if (CachedDataTable.Num() == 0)
 		{
@@ -227,15 +253,32 @@ TSharedRef<SVerticalBox> FDataTableEditor::CreateContentBox()
 		];
 }
 
-void FDataTableEditor::OnDataTableReloaded() 
+TSharedRef<SWidget> FDataTableEditor::CreateRowEditorBox()
 {
-	CachedDataTable.Reset();
+	return SNew(SRowEditor, DataTable.Get());
+}
 
-	TSharedRef<SVerticalBox> ContentBox = CreateContentBox();
-	
-	GridPanelOwner->SetContent(ContentBox);
+TSharedRef<SDockTab> FDataTableEditor::SpawnTab_RowEditor(const FSpawnTabArgs& Args)
+{
+	check(Args.GetTabId().TabType == RowEditorTabId);
+
+	return SNew(SDockTab)
+		.Icon(FEditorStyle::GetBrush("DataTableEditor.Tabs.Properties"))
+		.Label(LOCTEXT("RowEditorTitle", "Row Editor"))
+		.TabColorScale(GetTabColorScale())
+		[
+			SNew(SBorder)
+			.Padding(2)
+			.VAlign(VAlign_Top)
+			.HAlign(HAlign_Fill)
+			.BorderImage(FEditorStyle::GetBrush("ToolPanel.GroupBorder"))
+			[
+				CreateRowEditorBox()
+			]
+		];
 }
 
+
 TSharedRef<SDockTab> FDataTableEditor::SpawnTab_DataTable( const FSpawnTabArgs& Args )
 {
 	check( Args.GetTabId().TabType == DataTableTabId );
diff --git a/Engine/Source/Editor/DataTableEditor/Private/DataTableEditor.h b/Engine/Source/Editor/DataTableEditor/Private/DataTableEditor.h
--- a/Engine/Source/Editor/DataTableEditor/Private/DataTableEditor.h
+++ b/Engine/Source/Editor/DataTableEditor/Private/DataTableEditor.h
@@ -5,11 +5,15 @@
 #include "IDataTableEditor.h"
 #include "Toolkits/AssetEditorToolkit.h"
 
+#include "Kismet2/StructureEditorUtils.h"
+#include "DataTableEditorUtils.h"
+
 class UDataTable;
 
 /** Viewer/editor for a DataTable */
-class FDataTableEditor :
-	public IDataTableEditor
+class FDataTableEditor : public IDataTableEditor
+	, public FStructureEditorUtils::INotifyOnStructChanged
+	, public FDataTableEditorUtils::INotifyOnDataTableChanged
 {
 
 public:
@@ -41,24 +45,30 @@ public:
 	virtual FString GetWorldCentricTabPrefix() const override;
 	virtual FLinearColor GetWorldCentricTabColorScale() const override;
 
-	/* Called when DataTable was reloaded */
-	virtual void OnDataTableReloaded() override;
+	// INotifyOnStructChanged
+	virtual void PreChange(const class UUserDefinedStruct* Struct, FStructureEditorUtils::EStructureEditorChangeInfo Info) override;
+	virtual void PostChange(const class UUserDefinedStruct* Struct, FStructureEditorUtils::EStructureEditorChangeInfo Info) override;
 
+	// INotifyOnDataTableChanged
+	virtual void PreChange(const UDataTable* Changed, FDataTableEditorUtils::EDataTableChangeInfo Info) override;
+	virtual void PostChange(const UDataTable* Changed, FDataTableEditorUtils::EDataTableChangeInfo Info) override;
 
 private:
 
-	/** Called when an object is reimported */
-	void OnPostReimport(UObject* InObject, bool bSuccess);
-
 	void ReloadVisibleData();
 
 	void OnSearchTextChanged(const FText& SearchText);
 
 	TSharedRef<SVerticalBox> CreateContentBox();
 
+	TSharedRef<SWidget> CreateRowEditorBox();
+
 	/**	Spawns the tab with the data table inside */
 	TSharedRef<SDockTab> SpawnTab_DataTable( const FSpawnTabArgs& Args );
 
+	/**	Spawns the tab with the Row Editor inside */
+	TSharedRef<SDockTab> SpawnTab_RowEditor(const FSpawnTabArgs& Args);
+
 private:
 
 	/** Cached table data */
@@ -77,8 +87,11 @@ private:
 	TSharedPtr<SBorder>	 GridPanelOwner;
 
 	/* The DataTable that is active in the editor */
-	UDataTable*	DataTable;
+	TAssetPtr<UDataTable> DataTable;
 
 	/**	The tab id for the data table tab */
 	static const FName DataTableTabId;
+
+	/**	The tab id for the row editor tab */
+	static const FName RowEditorTabId;
 };
diff --git a/Engine/Source/Editor/DataTableEditor/Private/SRowEditor.cpp b/Engine/Source/Editor/DataTableEditor/Private/SRowEditor.cpp
new file mode 100644
--- /dev/null
+++ b/Engine/Source/Editor/DataTableEditor/Private/SRowEditor.cpp
@@ -0,0 +1,346 @@
+// Copyright 1998-2014 Epic Games, Inc. All Rights Reserved.
+
+#include "DataTableEditorPrivatePCH.h"
+
+#include "SRowEditor.h"
+#include "PropertyEditorModule.h"
+#include "IStructureDetailsView.h"
+#include "IDetailsView.h"
+#include "DataTableEditorUtils.h"
+
+#define LOCTEXT_NAMESPACE "SRowEditor"
+
+class FStructFromDataTable : public FStructOnScope
+{
+	TWeakObjectPtr<UDataTable> DataTable;
+	FName RowName;
+
+public:
+	FStructFromDataTable(UDataTable* InDataTable, FName InRowName)
+		: FStructOnScope()
+		, DataTable(InDataTable)
+		, RowName(InRowName)
+	{}
+
+	virtual uint8* GetStructMemory() override
+	{
+		return (DataTable.IsValid() && !RowName.IsNone()) ? DataTable->FindRowUnchecked(RowName) : NULL;
+	}
+
+	virtual const uint8* GetStructMemory() const override
+	{
+		return (DataTable.IsValid() && !RowName.IsNone()) ? DataTable->FindRowUnchecked(RowName) : NULL;
+	}
+
+	virtual const UScriptStruct* GetStruct() const override
+	{
+		return DataTable.IsValid() ? DataTable->RowStruct : NULL;
+	}
+
+	virtual bool IsValid() const override
+	{
+		return !RowName.IsNone()
+			&& DataTable.IsValid()
+			&& DataTable->RowStruct
+			&& DataTable->FindRowUnchecked(RowName);
+	}
+
+	virtual void Destroy() override
+	{
+		DataTable = NULL;
+		RowName = NAME_None;
+	}
+
+	FName GetRowName() const
+	{
+		return RowName;
+	}
+};
+
+SRowEditor::SRowEditor() : SCompoundWidget()
+{
+}
+
+SRowEditor::~SRowEditor()
+{
+}
+
+void SRowEditor::OnFinishedChangingProperties(const FPropertyChangedEvent& PropertyChangedEvent)
+{
+	check(PropertyChangedEvent.MemberProperty
+		&& PropertyChangedEvent.MemberProperty->GetOwnerStruct()
+		&& (PropertyChangedEvent.MemberProperty->GetOwnerStruct() == GetScriptStruct()));
+
+	// TODO:
+	// FNotifyHook won't be called, because the struct has no "property-path" from DT object. THe struct is seen as detached.
+	// But the notification must be called somewhere.
+
+	FDataTableEditorUtils::BroadcastPostChange(DataTable.Get(), FDataTableEditorUtils::EDataTableChangeInfo::RowDataPostChangeOnly);
+
+	if (DataTable.IsValid())
+	{
+		DataTable->MarkPackageDirty();
+	}
+}
+
+void SRowEditor::PreChange(const class UUserDefinedStruct* Struct, FStructureEditorUtils::EStructureEditorChangeInfo Info)
+{
+	if (Struct && (GetScriptStruct() == Struct))
+	{
+		CleanBeforeChange();
+	}
+}
+
+void SRowEditor::PostChange(const class UUserDefinedStruct* Struct, FStructureEditorUtils::EStructureEditorChangeInfo Info)
+{
+	if (Struct && (GetScriptStruct() == Struct))
+	{
+		Restore();
+	}
+}
+
+void SRowEditor::PreChange(const UDataTable* Changed, FDataTableEditorUtils::EDataTableChangeInfo Info)
+{
+	if ((Changed == DataTable.Get()) && (FDataTableEditorUtils::EDataTableChangeInfo::RowList == Info))
+	{
+		CleanBeforeChange();
+	}
+}
+
+void SRowEditor::PostChange(const UDataTable* Changed, FDataTableEditorUtils::EDataTableChangeInfo Info)
+{
+	FStringAssetReference::InvalidateTag(); // Should be removed after UE-5615 is fixed
+	if ((Changed == DataTable.Get()) && (FDataTableEditorUtils::EDataTableChangeInfo::RowList == Info))
+	{
+		RefreshNameList();
+		Restore();
+	}
+}
+
+void SRowEditor::CleanBeforeChange()
+{
+	if (StructureDetailsView.IsValid())
+	{
+		StructureDetailsView->SetStructureData(NULL);
+	}
+	if (CurrentRow.IsValid())
+	{
+		CurrentRow->Destroy();
+		CurrentRow.Reset();
+	}
+}
+
+void SRowEditor::RefreshNameList()
+{
+	CachedRowNames.Empty();
+	if (DataTable.IsValid())
+	{
+		auto RowNames = DataTable->GetRowNames();
+		for (auto RowName : RowNames)
+		{
+			CachedRowNames.Add(MakeShareable(new FName(RowName)));
+		}
+	}
+}
+
+void SRowEditor::Restore()
+{
+	if (SelectedName.IsValid())
+	{
+		auto CurrentName = *SelectedName;
+		SelectedName = NULL;
+		for (auto Element : CachedRowNames)
+		{
+			if (*Element == CurrentName)
+			{
+				SelectedName = Element;
+				break;
+			}
+		}
+	}
+
+	if (!SelectedName.IsValid() && CachedRowNames.Num() && CachedRowNames[0].IsValid())
+	{
+		SelectedName = CachedRowNames[0];
+	}
+
+	if (RowComboBox.IsValid())
+	{
+		RowComboBox->SetSelectedItem(SelectedName);
+	}
+
+	auto FinalName = SelectedName.IsValid() ? *SelectedName : NAME_None;
+	CurrentRow = MakeShareable(new FStructFromDataTable(DataTable.Get(), FinalName));
+	if (StructureDetailsView.IsValid())
+	{
+		StructureDetailsView->SetStructureData(CurrentRow);
+	}
+}
+
+UScriptStruct* SRowEditor::GetScriptStruct() const
+{
+	return DataTable.IsValid() ? DataTable->RowStruct : NULL;
+}
+
+FName SRowEditor::GetCurrentName() const
+{
+	return SelectedName.IsValid() ? *SelectedName : NAME_None;
+}
+
+FText SRowEditor::GetCurrentNameAsText() const
+{
+	return FText::FromName(GetCurrentName());
+}
+
+FString SRowEditor::GetStructureDisplayName() const
+{
+	const auto Struct = DataTable.IsValid() ? DataTable->RowStruct : NULL;
+	return Struct 
+		? Struct->GetDisplayNameText().ToString()
+		: LOCTEXT("Error_UnknownStruct", "Error: Unknown Struct").ToString();
+}
+
+TSharedRef<SWidget> SRowEditor::OnGenerateWidget(TSharedPtr<FName> InItem)
+{
+	return SNew(STextBlock).Text(FText::FromName(InItem.IsValid() ? *InItem : NAME_None));
+}
+
+void SRowEditor::OnSelectionChanged(TSharedPtr<FName> InItem, ESelectInfo::Type InSeletionInfo)
+{
+	if (InItem != SelectedName)
+	{
+		CleanBeforeChange();
+
+		SelectedName = InItem;
+
+		Restore();
+	}
+}
+
+FReply SRowEditor::OnAddClicked()
+{
+	if (DataTable.IsValid())
+	{
+		FName NewName = DataTable->MakeValidName(TEXT("NewRow"));
+		const TArray<FName> ExisitngNames = DataTable->GetRowNames();
+		while (ExisitngNames.Contains(NewName))
+		{
+			NewName.SetNumber(NewName.GetNumber() + 1);
+		}
+		SelectedName = MakeShareable(new FName(NewName));
+		FDataTableEditorUtils::AddRow(DataTable.Get(), NewName);
+	}
+	return FReply::Handled();
+}
+
+FReply SRowEditor::OnRemoveClicked()
+{
+	if (DataTable.IsValid())
+	{
+		const FName RowToRemove = GetCurrentName();
+		FDataTableEditorUtils::RemoveRow(DataTable.Get(), RowToRemove);
+	}
+	return FReply::Handled();
+}
+
+void SRowEditor::OnRowRenamed(const FText& Text, ETextCommit::Type CommitType)
+{
+	if (!GetCurrentNameAsText().EqualTo(Text) && DataTable.IsValid())
+	{
+		const FName OldName = GetCurrentName();
+		const FName NewName = DataTable->MakeValidName(Text.ToString());
+		SelectedName = MakeShareable(new FName(NewName));
+		FDataTableEditorUtils::RenameRow(DataTable.Get(), OldName, NewName);
+	}
+}
+
+void SRowEditor::Construct(const FArguments& InArgs, UDataTable* Changed)
+{
+	DataTable = Changed;
+	{
+		FPropertyEditorModule& PropertyModule = FModuleManager::LoadModuleChecked<FPropertyEditorModule>("PropertyEditor");
+		FDetailsViewArgs ViewArgs;
+		ViewArgs.bAllowSearch = false;
+		ViewArgs.bHideSelectionTip = false;
+		ViewArgs.bShowActorLabel = false;
+
+		StructureDetailsView = PropertyModule.CreateStructureDetailView(ViewArgs, CurrentRow, false, LOCTEXT("RowValue", "Row Value").ToString());
+		StructureDetailsView->GetOnFinishedChangingPropertiesDelegate().AddSP(this, &SRowEditor::OnFinishedChangingProperties);
+	}
+
+	RefreshNameList();
+	Restore();
+	
+	ChildSlot
+	[
+		SNew(SVerticalBox)
+		+ SVerticalBox::Slot()
+		.AutoHeight()
+		[
+			SNew(SHorizontalBox)
+			+ SHorizontalBox::Slot()
+			.AutoWidth()
+			.Padding(2)
+			.VAlign(EVerticalAlignment::VAlign_Center)
+			[
+				SNew(STextBlock)
+				.Text(this, &SRowEditor::GetStructureDisplayName)
+			]
+			+ SHorizontalBox::Slot()
+			.AutoWidth()
+			.Padding(2)
+			[
+				SAssignNew(RowComboBox, SComboBox<TSharedPtr<FName>>)
+				.OptionsSource(&CachedRowNames)
+				.OnSelectionChanged(this, &SRowEditor::OnSelectionChanged)
+				.OnGenerateWidget(this, &SRowEditor::OnGenerateWidget)
+				.Content()
+				[
+					SNew(STextBlock).Text(this, &SRowEditor::GetCurrentNameAsText)
+				]
+			]
+			+ SHorizontalBox::Slot()
+			.AutoWidth()
+			.Padding(2)
+			[
+				SNew(SButton)
+				.OnClicked(this, &SRowEditor::OnAddClicked)
+				.Content()
+				[
+					SNew(STextBlock).Text(LOCTEXT("AddRow", "Add"))
+				]
+			]
+			+ SHorizontalBox::Slot()
+			.AutoWidth()
+			.Padding(2)
+			[
+				SNew(SButton)
+				.OnClicked(this, &SRowEditor::OnRemoveClicked)
+				.Content()
+				[
+					SNew(STextBlock).Text(LOCTEXT("RemoveRow", "Remove"))
+				]
+			]
+			+ SHorizontalBox::Slot()
+			.AutoWidth()
+			.Padding(2)
+			.VAlign(EVerticalAlignment::VAlign_Center)
+			[
+				SNew(STextBlock).Text(LOCTEXT("RenameRow", "Rename:"))
+			]
+			+ SHorizontalBox::Slot()
+			.AutoWidth()
+			.Padding(2)
+			[
+				SNew(SEditableTextBox)
+				.Text(this, &SRowEditor::GetCurrentNameAsText)
+				.OnTextCommitted(this, &SRowEditor::OnRowRenamed)
+			]
+		]
+		+ SVerticalBox::Slot()
+		.AutoHeight()
+		[
+			StructureDetailsView->GetWidget().ToSharedRef()
+		]
+	];
+}
\ No newline at end of file
diff --git a/Engine/Source/Editor/DataTableEditor/Private/SRowEditor.h b/Engine/Source/Editor/DataTableEditor/Private/SRowEditor.h
new file mode 100644
--- /dev/null
+++ b/Engine/Source/Editor/DataTableEditor/Private/SRowEditor.h
@@ -0,0 +1,59 @@
+// Copyright 1998-2014 Epic Games, Inc. All Rights Reserved.
+
+#pragma once
+
+#include "Kismet2/StructureEditorUtils.h"
+#include "DataTableEditorUtils.h"
+
+DECLARE_DELEGATE_OneParam(FOnRowModified, FName /*Row name*/);
+
+class SRowEditor : public SCompoundWidget
+	, public FStructureEditorUtils::INotifyOnStructChanged
+	, public FDataTableEditorUtils::INotifyOnDataTableChanged
+{
+public:
+	SLATE_BEGIN_ARGS(SRowEditor) {}
+	SLATE_END_ARGS()
+
+	SRowEditor();
+	virtual ~SRowEditor();
+
+	// INotifyOnStructChanged
+	virtual void PreChange(const class UUserDefinedStruct* Struct, FStructureEditorUtils::EStructureEditorChangeInfo Info) override;
+	virtual void PostChange(const class UUserDefinedStruct* Struct, FStructureEditorUtils::EStructureEditorChangeInfo Info) override;
+
+	// INotifyOnDataTableChanged
+	virtual void PreChange(const UDataTable* Changed, FDataTableEditorUtils::EDataTableChangeInfo Info) override;
+	virtual void PostChange(const UDataTable* Changed, FDataTableEditorUtils::EDataTableChangeInfo Info) override;
+
+private:
+
+	TArray<TSharedPtr<FName>> CachedRowNames;
+	TSharedPtr<FStructOnScope> CurrentRow;
+	TAssetPtr<UDataTable> DataTable; // weak obj ptr couldn't handle reimporting
+	TSharedPtr<class IStructureDetailsView> StructureDetailsView;
+	TSharedPtr<FName> SelectedName;
+	TSharedPtr<SComboBox<TSharedPtr<FName>>> RowComboBox;
+
+	void RefreshNameList();
+	void CleanBeforeChange();
+	void Restore();
+
+	UScriptStruct* GetScriptStruct() const;
+
+	FName GetCurrentName() const;
+	FText GetCurrentNameAsText() const;
+	FString GetStructureDisplayName() const;
+	TSharedRef<SWidget> OnGenerateWidget(TSharedPtr<FName> InItem);
+	void OnSelectionChanged(TSharedPtr<FName> InItem, ESelectInfo::Type InSeletionInfo);
+
+	void OnFinishedChangingProperties(const struct FPropertyChangedEvent& PropertyChangedEvent);
+
+	FReply OnAddClicked();
+	FReply OnRemoveClicked();
+	void OnRowRenamed(const FText& Text, ETextCommit::Type CommitType);
+
+public:
+
+	void Construct(const FArguments& InArgs, UDataTable* Changed);
+};
\ No newline at end of file
diff --git a/Engine/Source/Editor/DataTableEditor/Public/IDataTableEditor.h b/Engine/Source/Editor/DataTableEditor/Public/IDataTableEditor.h
--- a/Engine/Source/Editor/DataTableEditor/Public/IDataTableEditor.h
+++ b/Engine/Source/Editor/DataTableEditor/Public/IDataTableEditor.h
@@ -13,10 +13,6 @@ class IDataTableEditor : public FAssetEditorToolkit
 public:
 
 	// ...
-
-	/** Called after a data table was reloaded */
-	virtual void OnDataTableReloaded(){}
-
 };
 
 
diff --git a/Engine/Source/Editor/GraphEditor/Private/KismetPins/SGraphPinDataTableRowName.cpp b/Engine/Source/Editor/GraphEditor/Private/KismetPins/SGraphPinDataTableRowName.cpp
new file mode 100644
--- /dev/null
+++ b/Engine/Source/Editor/GraphEditor/Private/KismetPins/SGraphPinDataTableRowName.cpp
@@ -0,0 +1,47 @@
+// Copyright 1998-2014 Epic Games, Inc. All Rights Reserved.
+
+#pragma once
+#include "GraphEditorCommon.h"
+#include "SGraphPinDataTableRowName.h"
+
+void SGraphPinDataTableRowName::Construct(const FArguments& InArgs, UEdGraphPin* InGraphPinObj, class UDataTable* InDataTable)
+{
+	DataTable = InDataTable;
+	RefreshNameList();
+	SGraphPinNameList::Construct(SGraphPinNameList::FArguments(), InGraphPinObj, NameList);
+}
+
+SGraphPinDataTableRowName::SGraphPinDataTableRowName()
+{
+}
+
+SGraphPinDataTableRowName::~SGraphPinDataTableRowName()
+{
+}
+
+void SGraphPinDataTableRowName::PreChange(const UDataTable* Changed, FDataTableEditorUtils::EDataTableChangeInfo Info)
+{
+}
+
+void SGraphPinDataTableRowName::PostChange(const UDataTable* Changed, FDataTableEditorUtils::EDataTableChangeInfo Info)
+{
+	//FStringAssetReference::InvalidateTag(); // Should be removed after UE-5615 is fixed
+	if (Changed && (Changed == DataTable.Get()) && (FDataTableEditorUtils::EDataTableChangeInfo::RowList == Info))
+	{
+		RefreshNameList();
+	}
+}
+
+void SGraphPinDataTableRowName::RefreshNameList()
+{
+	NameList.Empty();
+	if (DataTable.IsValid())
+	{
+		auto Names = DataTable->GetRowNames();
+		for (auto Name : Names)
+		{
+			TSharedPtr<FName> RowNameItem = MakeShareable(new FName(Name));
+			NameList.Add(RowNameItem);
+		}
+	}
+}
\ No newline at end of file
diff --git a/Engine/Source/Editor/GraphEditor/Private/KismetPins/SGraphPinNameList.cpp b/Engine/Source/Editor/GraphEditor/Private/KismetPins/SGraphPinNameList.cpp
--- a/Engine/Source/Editor/GraphEditor/Private/KismetPins/SGraphPinNameList.cpp
+++ b/Engine/Source/Editor/GraphEditor/Private/KismetPins/SGraphPinNameList.cpp
@@ -13,16 +13,12 @@ void SGraphPinNameList::Construct(const FArguments& InArgs, UEdGraphPin* InGraph
 
 TSharedRef<SWidget>	SGraphPinNameList::GetDefaultValueWidget()
 {
-	//Get list of name indexes
-	TArray< TSharedPtr<int32> > ComboItems;
-	GenerateComboBoxIndexes( ComboItems );
-
 	TSharedPtr<FName> CurrentlySelectedName;
-	if (NameList.Num() > 0)
+
+	if (GraphPinObj)
 	{
 		// Preserve previous selection, or set to first in list
 		FName PreviousSelection = FName(*GraphPinObj->GetDefaultAsString());
-		CurrentlySelectedName = NameList[0];
 		for (TSharedPtr<FName> ListNamePtr : NameList)
 		{
 			if (PreviousSelection == *ListNamePtr.Get())
@@ -31,14 +27,8 @@ TSharedRef<SWidget>	SGraphPinNameList::GetDefaultValueWidget()
 				break;
 			}
 		}
-		check(CurrentlySelectedName.IsValid());
-		// Ensure the pin value is synchronised with the graph pin widget
-		if( PreviousSelection != *CurrentlySelectedName.Get() )
-		{
-			GraphPinObj->GetSchema()->TrySetDefaultValue(*GraphPinObj, *CurrentlySelectedName->ToString());
-		}
 	}
-
+	
 	// Create widget
 	return SAssignNew(ComboBox, SNameComboBox)
 		.ContentPadding(FMargin(6.0f, 2.0f))
@@ -49,29 +39,12 @@ TSharedRef<SWidget>	SGraphPinNameList::GetDefaultValueWidget()
 		;
 }
 
-FString SGraphPinNameList::OnGetFriendlyName(int32 NameIndex)
-{
-	check(NameIndex < NameList.Num());
-	return NameList[NameIndex]->ToString();
-}
-
 void SGraphPinNameList::ComboBoxSelectionChanged(TSharedPtr<FName> NameItem, ESelectInfo::Type SelectInfo )
 {
-	check(NameItem.IsValid());
-	GraphPinObj->GetSchema()->TrySetDefaultValue(*GraphPinObj, *NameItem->ToString());
-}
-
-FString SGraphPinNameList::OnGetText() const 
-{
-	FString SelectedString = GraphPinObj->GetDefaultAsString();
-	return SelectedString;
-}
-
-void SGraphPinNameList::GenerateComboBoxIndexes( TArray< TSharedPtr<int32> >& OutComboBoxIndexes )
-{
-	for (int32 NameIndex = 0; NameIndex < NameList.Num(); NameIndex++)
+	FName Name = NameItem.IsValid() ? *NameItem : NAME_None;
+	if (auto Schema = (GraphPinObj ? GraphPinObj->GetSchema() : NULL))
 	{
-		TSharedPtr<int32> NameIdxPtr(new int32(NameIndex));
-		OutComboBoxIndexes.Add(NameIdxPtr);
+		Schema->TrySetDefaultValue(*GraphPinObj, *Name.ToString());
 	}
 }
+
diff --git a/Engine/Source/Editor/GraphEditor/Public/SGraphPinDataTableRowName.h b/Engine/Source/Editor/GraphEditor/Public/SGraphPinDataTableRowName.h
new file mode 100644
--- /dev/null
+++ b/Engine/Source/Editor/GraphEditor/Public/SGraphPinDataTableRowName.h
@@ -0,0 +1,27 @@
+// Copyright 1998-2014 Epic Games, Inc. All Rights Reserved.
+
+#pragma once
+#include "SGraphPinNameList.h"
+#include "DataTableEditorUtils.h"
+
+class GRAPHEDITOR_API SGraphPinDataTableRowName : public SGraphPinNameList, public FDataTableEditorUtils::INotifyOnDataTableChanged
+{
+public:
+	SLATE_BEGIN_ARGS(SGraphPinDataTableRowName) {}
+	SLATE_END_ARGS()
+
+	void Construct(const FArguments& InArgs, UEdGraphPin* InGraphPinObj, class UDataTable* InDataTable);
+
+	SGraphPinDataTableRowName();
+	virtual ~SGraphPinDataTableRowName();
+
+	// FDataTableEditorUtils::INotifyOnDataTableChanged
+	virtual void PreChange(const UDataTable* Changed, FDataTableEditorUtils::EDataTableChangeInfo Info) override;
+	virtual void PostChange(const UDataTable* Changed, FDataTableEditorUtils::EDataTableChangeInfo Info) override;
+
+protected:
+
+	void RefreshNameList();
+
+	TAssetPtr<class UDataTable> DataTable;
+};
diff --git a/Engine/Source/Editor/GraphEditor/Public/SGraphPinNameList.h b/Engine/Source/Editor/GraphEditor/Public/SGraphPinNameList.h
--- a/Engine/Source/Editor/GraphEditor/Public/SGraphPinNameList.h
+++ b/Engine/Source/Editor/GraphEditor/Public/SGraphPinNameList.h
@@ -21,22 +21,6 @@ protected:
 	 */
 	virtual TSharedRef<SWidget>	GetDefaultValueWidget() override;
 
-private:
-
-	/**
-	 *	Function to get current string associated with the combo box selection
-	 *
-	 *	@return currently selected string
-	 */
-	FString OnGetText() const;
-
-	/**
-	 *	Function to generate the list of indexes from the enum object
-	 *
-	 *	@param OutComboBoxIndexes - Int array reference to store the list of indexes
-	 */
-	void GenerateComboBoxIndexes( TArray< TSharedPtr<int32> >& OutComboBoxIndexes );
-
 	/**
 	 *	Function to set the newly selected index
 	 *
@@ -45,13 +29,6 @@ private:
 	 */
 	void ComboBoxSelectionChanged(TSharedPtr<FName> StringItem, ESelectInfo::Type SelectInfo);
 
-	/**
-	 * Returns the friendly name of the enum at index EnumIndex
-	 *
-	 * @param EnumIndex	- The index of the enum to return the friendly name for
-	 */
-	FString OnGetFriendlyName(int32 EnumIndex);
-
 	TSharedPtr<class SNameComboBox>	ComboBox;
 
 	/** The actual list of FName values to choose from */
diff --git a/Engine/Source/Editor/Kismet/Private/UserDefinedEnumEditor.cpp b/Engine/Source/Editor/Kismet/Private/UserDefinedEnumEditor.cpp
--- a/Engine/Source/Editor/Kismet/Private/UserDefinedEnumEditor.cpp
+++ b/Engine/Source/Editor/Kismet/Private/UserDefinedEnumEditor.cpp
@@ -166,13 +166,10 @@ void FEnumDetails::CustomizeDetails( IDetailLayoutBuilder& DetailLayout )
 		Layout = MakeShareable( new FUserDefinedEnumLayout(TargetEnum.Get()) );
 		InputsCategory.AddCustomBuilder( Layout.ToSharedRef() );
 	}
-
-	FEnumEditorUtils::FEnumEditorManager::Get().AddListener(this);
 }
 
 FEnumDetails::~FEnumDetails()
 {
-	FEnumEditorUtils::FEnumEditorManager::Get().RemoveListener(this);
 }
 
 void FEnumDetails::OnForceRefresh()
@@ -183,11 +180,11 @@ void FEnumDetails::OnForceRefresh()
 	}
 }
 
-void FEnumDetails::PreChange(const class UUserDefinedEnum* Enum)
+void FEnumDetails::PreChange(const class UUserDefinedEnum* Enum, FEnumEditorUtils::EEnumEditorChangeInfo Info)
 {
 }
 
-void FEnumDetails::PostChange(const class UUserDefinedEnum* Enum)
+void FEnumDetails::PostChange(const class UUserDefinedEnum* Enum, FEnumEditorUtils::EEnumEditorChangeInfo Info)
 {
 	if (Enum && (TargetEnum.Get() == Enum))
 	{
diff --git a/Engine/Source/Editor/Kismet/Private/UserDefinedEnumEditor.h b/Engine/Source/Editor/Kismet/Private/UserDefinedEnumEditor.h
--- a/Engine/Source/Editor/Kismet/Private/UserDefinedEnumEditor.h
+++ b/Engine/Source/Editor/Kismet/Private/UserDefinedEnumEditor.h
@@ -69,8 +69,8 @@ public:
 	void OnForceRefresh();
 
 	/** FEnumEditorUtils::INotifyOnEnumChanged */
-	virtual void PreChange(const class UUserDefinedEnum* Struct) override;
-	virtual void PostChange(const class UUserDefinedEnum* Enum) override;
+	virtual void PreChange(const class UUserDefinedEnum* Enum, FEnumEditorUtils::EEnumEditorChangeInfo Info) override;
+	virtual void PostChange(const class UUserDefinedEnum* Enum, FEnumEditorUtils::EEnumEditorChangeInfo Info) override;
 
 private:
 	/** Handles new enum element request */
diff --git a/Engine/Source/Editor/Kismet/Private/UserDefinedStructureEditor.cpp b/Engine/Source/Editor/Kismet/Private/UserDefinedStructureEditor.cpp
--- a/Engine/Source/Editor/Kismet/Private/UserDefinedStructureEditor.cpp
+++ b/Engine/Source/Editor/Kismet/Private/UserDefinedStructureEditor.cpp
@@ -21,7 +21,6 @@ public:
 	FStructureDefaultValueView(UUserDefinedStruct* EditedStruct) 
 		: UserDefinedStruct(EditedStruct)
 	{
-		FStructureEditorUtils::FStructEditorManager::Get().AddListener(this);
 	}
 
 	void Initialize()
@@ -41,7 +40,6 @@ public:
 
 	virtual ~FStructureDefaultValueView()
 	{
-		FStructureEditorUtils::FStructEditorManager::Get().RemoveListener(this);
 	}
 
 	void OnFinishedChangingProperties(const FPropertyChangedEvent& PropertyChangedEvent)
@@ -83,7 +81,7 @@ public:
 		return StructureDetailsView.IsValid() ? StructureDetailsView->GetWidget() : NULL;
 	}
 
-	virtual void PreChange(const class UUserDefinedStruct* Struct) override
+	virtual void PreChange(const class UUserDefinedStruct* Struct, FStructureEditorUtils::EStructureEditorChangeInfo Info) override
 	{
 		if (Struct && (GetUserDefinedStruct() == Struct))
 		{
@@ -99,7 +97,7 @@ public:
 		}
 	}
 
-	virtual void PostChange(const class UUserDefinedStruct* Struct) override
+	virtual void PostChange(const class UUserDefinedStruct* Struct, FStructureEditorUtils::EStructureEditorChangeInfo Info) override
 	{
 		if (Struct && (GetUserDefinedStruct() == Struct))
 		{
@@ -132,7 +130,6 @@ public:
 
 	~FUserDefinedStructureDetails()
 	{
-		FStructureEditorUtils::FStructEditorManager::Get().RemoveListener(this);
 	}
 
 	UUserDefinedStruct* GetUserDefinedStruct()
@@ -153,8 +150,8 @@ public:
 	virtual void CustomizeDetails(class IDetailLayoutBuilder& DetailLayout) override;
 
 	/** FStructureEditorUtils::INotifyOnStructChanged */
-	virtual void PreChange(const class UUserDefinedStruct* Struct) override {}
-	virtual void PostChange(const class UUserDefinedStruct* Struct) override;
+	virtual void PreChange(const class UUserDefinedStruct* Struct, FStructureEditorUtils::EStructureEditorChangeInfo Info) override {}
+	virtual void PostChange(const class UUserDefinedStruct* Struct, FStructureEditorUtils::EStructureEditorChangeInfo Info) override;
 
 private:
 	TWeakObjectPtr<UUserDefinedStruct> UserDefinedStruct;
@@ -935,12 +932,10 @@ void FUserDefinedStructureDetails::CustomizeDetails(class IDetailLayoutBuilder&
 		Layout = MakeShareable(new FUserDefinedStructureLayout(SharedThis(this)));
 		StructureCategory.AddCustomBuilder(Layout.ToSharedRef());
 	}
-
-	FStructureEditorUtils::FStructEditorManager::Get().AddListener(this);
 }
 
 /** FStructureEditorUtils::INotifyOnStructChanged */
-void FUserDefinedStructureDetails::PostChange(const class UUserDefinedStruct* Struct)
+void FUserDefinedStructureDetails::PostChange(const class UUserDefinedStruct* Struct, FStructureEditorUtils::EStructureEditorChangeInfo Info)
 {
 	if (Struct && (GetUserDefinedStruct() == Struct))
 	{
diff --git a/Engine/Source/Editor/UnrealEd/Private/AutoReimportManager.cpp b/Engine/Source/Editor/UnrealEd/Private/AutoReimportManager.cpp
--- a/Engine/Source/Editor/UnrealEd/Private/AutoReimportManager.cpp
+++ b/Engine/Source/Editor/UnrealEd/Private/AutoReimportManager.cpp
@@ -4,7 +4,6 @@
 #include "AssetToolsModule.h"
 
 #include "Developer/DirectoryWatcher/Public/DirectoryWatcherModule.h"
-#include "Editor/DataTableEditor/Public/IDataTableEditor.h"
 #include "Engine/CurveTable.h"
 #include "Engine/DataTable.h"
 
@@ -353,13 +352,7 @@ bool UAutoReimportManager::AutoReimportDataTable( UObject* Object, const FString
 			if (JustFileNameMatches(DataTable->ImportPath, FileName))
 			{
 				DataTable->ImportPath = FReimportManager::SanitizeImportFilename(FileName, DataTable);
-				if (ReimportFile(DataTable, FileName))
-				{
-					if (IDataTableEditor* Editor =  static_cast<IDataTableEditor*>(FAssetEditorManager::Get().FindEditorForAsset(DataTable, false)))
-					{
-						Editor->OnDataTableReloaded();
-					}
-				}
+				ReimportFile(DataTable, FileName);
 
 				return true;
 			}
diff --git a/Engine/Source/Editor/UnrealEd/Private/DataTableEditorUtils.cpp b/Engine/Source/Editor/UnrealEd/Private/DataTableEditorUtils.cpp
new file mode 100644
--- /dev/null
+++ b/Engine/Source/Editor/UnrealEd/Private/DataTableEditorUtils.cpp
@@ -0,0 +1,93 @@
+// Copyright 1998-2014 Epic Games, Inc. All Rights Reserved.
+
+#include "UnrealEd.h"
+#include "DataTableEditorUtils.h"
+#include "K2Node_GetDataTableRow.h"
+
+FDataTableEditorUtils::FDataTableEditorManager& FDataTableEditorUtils::FDataTableEditorManager::Get()
+{
+	static TSharedRef< FDataTableEditorManager > EditorManager(new FDataTableEditorManager());
+	return *EditorManager;
+}
+
+bool FDataTableEditorUtils::RemoveRow(UDataTable* DataTable, FName Name)
+{
+	bool bResult = false;
+	if (DataTable && DataTable->RowStruct)
+	{
+		BroadcastPreChange(DataTable, EDataTableChangeInfo::RowList);
+		DataTable->Modify();
+		uint8* RowData = NULL;
+		const bool bRemoved = DataTable->RowMap.RemoveAndCopyValue(Name, RowData);
+		if (bRemoved && RowData)
+		{
+			DataTable->RowStruct->DestroyScriptStruct(RowData);
+			FMemory::Free(RowData);
+			bResult = true;
+		}
+		BroadcastPostChange(DataTable, EDataTableChangeInfo::RowList);
+	}
+	return bResult;
+}
+
+uint8* FDataTableEditorUtils::AddRow(UDataTable* DataTable, FName RowName)
+{
+	if (!DataTable || (RowName == NAME_None) || (DataTable->RowMap.Find(RowName) != NULL) || !DataTable->RowStruct)
+	{
+		return NULL;
+	}
+	BroadcastPreChange(DataTable, EDataTableChangeInfo::RowList);
+	DataTable->Modify();
+	// Allocate data to store information, using UScriptStruct to know its size
+	uint8* RowData = (uint8*)FMemory::Malloc(DataTable->RowStruct->PropertiesSize);
+	DataTable->RowStruct->InitializeScriptStruct(RowData);
+	// And be sure to call DestroyScriptStruct later
+
+	if (auto UDStruct = Cast<const UUserDefinedStruct>(DataTable->RowStruct))
+	{
+		UDStruct->InitializeDefaultValue(RowData);
+	}
+
+	// Add to row map
+	DataTable->RowMap.Add(RowName, RowData);
+	BroadcastPostChange(DataTable, EDataTableChangeInfo::RowList);
+	return RowData;
+}
+
+bool FDataTableEditorUtils::RenameRow(UDataTable* DataTable, FName OldName, FName NewName)
+{
+	bool bResult = false;
+	if (DataTable)
+	{
+		BroadcastPreChange(DataTable, EDataTableChangeInfo::RowList);
+		DataTable->Modify();
+
+		uint8* RowData = NULL;
+		const bool bValidnewName = (NewName != NAME_None) && !DataTable->RowMap.Find(NewName);
+		const bool bRemoved = bValidnewName && DataTable->RowMap.RemoveAndCopyValue(OldName, RowData);
+		if (bRemoved)
+		{
+			DataTable->RowMap.FindOrAdd(NewName) = RowData;
+			bResult = true;
+		}
+		BroadcastPostChange(DataTable, EDataTableChangeInfo::RowList);
+	}
+	return bResult;
+}
+
+void FDataTableEditorUtils::BroadcastPreChange(UDataTable* DataTable, EDataTableChangeInfo Info)
+{
+	FDataTableEditorManager::Get().PreChange(DataTable, Info);
+}
+
+void FDataTableEditorUtils::BroadcastPostChange(UDataTable* DataTable, EDataTableChangeInfo Info)
+{
+	if (DataTable && (EDataTableChangeInfo::RowList == Info))
+	{
+		for (TObjectIterator<UK2Node_GetDataTableRow> It(RF_Transient | RF_PendingKill | RF_ClassDefaultObject); It; ++It)
+		{
+			It->OnDataTableRowListChanged(DataTable);
+		}
+	}
+	FDataTableEditorManager::Get().PostChange(DataTable, Info);
+}
diff --git a/Engine/Source/Editor/UnrealEd/Private/Factories/CSVImportFactory.cpp b/Engine/Source/Editor/UnrealEd/Private/Factories/CSVImportFactory.cpp
--- a/Engine/Source/Editor/UnrealEd/Private/Factories/CSVImportFactory.cpp
+++ b/Engine/Source/Editor/UnrealEd/Private/Factories/CSVImportFactory.cpp
@@ -12,6 +12,7 @@
 #include "Engine/DataTable.h"
 #include "Engine/CurveTable.h"
 #include "Engine/UserDefinedStruct.h"
+#include "DataTableEditorUtils.h"
 DEFINE_LOG_CATEGORY(LogCSVImportFactory);
 
 #define LOCTEXT_NAMESPACE "CSVImportFactory"
@@ -642,11 +643,14 @@ void UReimportDataTableFactory::SetReimportPaths( UObject* Obj, const TArray<FSt
 
 EReimportResult::Type UReimportDataTableFactory::Reimport( UObject* Obj )
 {	
-	if(Cast<UDataTable>(Obj))
+	auto Result = EReimportResult::Failed;
+	if(auto DataTable = Cast<UDataTable>(Obj))
 	{
-		return UCSVImportFactory::ReimportCSV(Obj) ? EReimportResult::Succeeded : EReimportResult::Failed;
+		FDataTableEditorUtils::BroadcastPreChange(DataTable, FDataTableEditorUtils::EDataTableChangeInfo::RowList);
+		Result = UCSVImportFactory::ReimportCSV(DataTable) ? EReimportResult::Succeeded : EReimportResult::Failed;
+		FDataTableEditorUtils::BroadcastPostChange(DataTable, FDataTableEditorUtils::EDataTableChangeInfo::RowList);
 	}
-	return EReimportResult::Failed;
+	return Result;
 }
 
 ////////////////////////////////////////////////////////////////////////////
diff --git a/Engine/Source/Editor/UnrealEd/Private/Kismet2/EnumEditorUtils.cpp b/Engine/Source/Editor/UnrealEd/Private/Kismet2/EnumEditorUtils.cpp
--- a/Engine/Source/Editor/UnrealEd/Private/Kismet2/EnumEditorUtils.cpp
+++ b/Engine/Source/Editor/UnrealEd/Private/Kismet2/EnumEditorUtils.cpp
@@ -17,9 +17,6 @@ struct FEnumEditorUtilsHelper
 
 //////////////////////////////////////////////////////////////////////////
 // FEnumEditorManager
-FEnumEditorUtils::FEnumEditorManager::FEnumEditorManager() : FListenerManager<UUserDefinedEnum>()
-{}
-
 FEnumEditorUtils::FEnumEditorManager& FEnumEditorUtils::FEnumEditorManager::Get()
 {
 	static TSharedRef< FEnumEditorManager > EnumEditorManager( new FEnumEditorManager() );
@@ -199,7 +196,7 @@ public:
 
 void FEnumEditorUtils::PrepareForChange(const UUserDefinedEnum* Enum)
 {
-	FEnumEditorManager::Get().PreChange(Enum);
+	FEnumEditorManager::Get().PreChange(Enum, EEnumEditorChangeInfo::Changed);
 }
 
 void FEnumEditorUtils::BroadcastChanges(const UUserDefinedEnum* Enum, const TArray<FName>& OldNames, bool bResolveData)
@@ -304,7 +301,7 @@ void FEnumEditorUtils::BroadcastChanges(const UUserDefinedEnum* Enum, const TArr
 		(*It)->BroadcastChanged();
 	}
 
-	FEnumEditorManager::Get().PostChange(Enum);
+	FEnumEditorManager::Get().PostChange(Enum, EEnumEditorChangeInfo::Changed);
 }
 
 int32 FEnumEditorUtils::ResolveEnumerator(const UEnum* Enum, FArchive& Ar, int32 EnumeratorIndex)
diff --git a/Engine/Source/Editor/UnrealEd/Private/Kismet2/StructureEditorUtils.cpp b/Engine/Source/Editor/UnrealEd/Private/Kismet2/StructureEditorUtils.cpp
--- a/Engine/Source/Editor/UnrealEd/Private/Kismet2/StructureEditorUtils.cpp
+++ b/Engine/Source/Editor/UnrealEd/Private/Kismet2/StructureEditorUtils.cpp
@@ -18,9 +18,6 @@
 
 //////////////////////////////////////////////////////////////////////////
 // FStructEditorManager
-FStructureEditorUtils::FStructEditorManager::FStructEditorManager() : FListenerManager<UUserDefinedStruct>()
-{}
-
 FStructureEditorUtils::FStructEditorManager& FStructureEditorUtils::FStructEditorManager::Get()
 {
 	static TSharedRef< FStructEditorManager > EditorManager( new FStructEditorManager() );
@@ -740,7 +737,7 @@ struct FReinstanceDataTableHelper
 
 void FStructureEditorUtils::BroadcastPreChange(UUserDefinedStruct* Struct)
 {
-	FStructureEditorUtils::FStructEditorManager::Get().PreChange(Struct);
+	FStructureEditorUtils::FStructEditorManager::Get().PreChange(Struct, EStructureEditorChangeInfo::Changed);
 	auto DataTables = FReinstanceDataTableHelper::GetTablesDependentOnStruct(Struct);
 	for (auto DataTable : DataTables)
 	{
@@ -750,17 +747,12 @@ void FStructureEditorUtils::BroadcastPreChange(UUserDefinedStruct* Struct)
 
 void FStructureEditorUtils::BroadcastPostChange(UUserDefinedStruct* Struct)
 {
-	FStructureEditorUtils::FStructEditorManager::Get().PostChange(Struct);
 	auto DataTables = FReinstanceDataTableHelper::GetTablesDependentOnStruct(Struct);
 	for (auto DataTable : DataTables)
 	{
 		DataTable->RestoreAfterStructChange();
-		auto DataTableEditor = static_cast<IDataTableEditor*>(FAssetEditorManager::Get().FindEditorForAsset(DataTable, false));
-		if (DataTableEditor)
-		{
-			DataTableEditor->OnDataTableReloaded();
-		}
 	}
+	FStructureEditorUtils::FStructEditorManager::Get().PostChange(Struct, EStructureEditorChangeInfo::Changed);
 }
 
 #undef LOCTEXT_NAMESPACE
diff --git a/Engine/Source/Editor/UnrealEd/Public/DataTableEditorUtils.h b/Engine/Source/Editor/UnrealEd/Public/DataTableEditorUtils.h
new file mode 100644
--- /dev/null
+++ b/Engine/Source/Editor/UnrealEd/Public/DataTableEditorUtils.h
@@ -0,0 +1,34 @@
+// Copyright 1998-2014 Epic Games, Inc. All Rights Reserved.
+#pragma once
+
+#include "Engine/DataTable.h"
+#include "ListenerManager.h"
+
+struct UNREALED_API FDataTableEditorUtils
+{
+	enum EDataTableChangeInfo
+	{
+		RowDataPostChangeOnly,
+		RowList,
+	};
+
+	class FDataTableEditorManager : public FListenerManager < UDataTable, EDataTableChangeInfo >
+	{
+		FDataTableEditorManager() {}
+	public:
+		UNREALED_API static FDataTableEditorManager& Get();
+
+		class UNREALED_API ListenerType : public InnerListenerType<FDataTableEditorManager>
+		{
+		};
+	};
+
+	typedef FDataTableEditorManager::ListenerType INotifyOnDataTableChanged;
+
+	static bool RemoveRow(UDataTable* DataTable, FName Name);
+	static uint8* AddRow(UDataTable* DataTable, FName RowName);
+	static bool RenameRow(UDataTable* DataTable, FName OldName, FName NewName);
+
+	static void BroadcastPreChange(UDataTable* DataTable, EDataTableChangeInfo Info);
+	static void BroadcastPostChange(UDataTable* DataTable, EDataTableChangeInfo Info);
+};
\ No newline at end of file
diff --git a/Engine/Source/Editor/UnrealEd/Public/Kismet2/EnumEditorUtils.h b/Engine/Source/Editor/UnrealEd/Public/Kismet2/EnumEditorUtils.h
--- a/Engine/Source/Editor/UnrealEd/Public/Kismet2/EnumEditorUtils.h
+++ b/Engine/Source/Editor/UnrealEd/Public/Kismet2/EnumEditorUtils.h
@@ -18,11 +18,20 @@ class UNREALED_API FEnumEditorUtils
 	static void CopyEnumeratorsWithoutMax(const UEnum* Enum, TArray<FName>& OutEnumNames);
 public:
 
-	class FEnumEditorManager : public FListenerManager<UUserDefinedEnum>
+	enum EEnumEditorChangeInfo
 	{
-		FEnumEditorManager();
+		Changed,
+	};
+
+	class FEnumEditorManager : public FListenerManager<UUserDefinedEnum, EEnumEditorChangeInfo>
+	{
+		FEnumEditorManager() {}
 	public:
 		UNREALED_API static FEnumEditorManager& Get();
+
+		class UNREALED_API ListenerType : public InnerListenerType<FEnumEditorManager>
+		{
+		};
 	};
 
 	typedef FEnumEditorManager::ListenerType INotifyOnEnumChanged;
diff --git a/Engine/Source/Editor/UnrealEd/Public/Kismet2/ListenerManager.h b/Engine/Source/Editor/UnrealEd/Public/Kismet2/ListenerManager.h
--- a/Engine/Source/Editor/UnrealEd/Public/Kismet2/ListenerManager.h
+++ b/Engine/Source/Editor/UnrealEd/Public/Kismet2/ListenerManager.h
@@ -2,25 +2,39 @@
 
 #pragma once
 
-template<typename TType>
+template<typename TType, typename TChangeInfo>
 class INotifyOnChanged
 {
 public:
-	virtual void PreChange(const TType* Changed) = 0;
-	virtual void PostChange(const TType* Changed) = 0;
+	virtual void PreChange(const TType* Changed, TChangeInfo ChangedType) = 0;
+	virtual void PostChange(const TType* Changed, TChangeInfo ChangedType) = 0;
 };
 
-template<typename TType>
+template<typename TType, typename TChangeInfo>
 class FListenerManager
 {
 public:
-	typedef INotifyOnChanged<TType> ListenerType;
+	typedef INotifyOnChanged<TType, TChangeInfo> BaseNotifyOnChanged;
+
+	template<typename TManagerType>
+	class InnerListenerType : public BaseNotifyOnChanged
+	{
+	public:
+		InnerListenerType()
+		{
+			TManagerType::Get().AddListener(this);
+		}
+		virtual ~InnerListenerType()
+		{
+			TManagerType::Get().RemoveListener(this);
+		}
+	};
 
 private:
-	TSet<ListenerType*> Listeners;
+	TSet<BaseNotifyOnChanged*> Listeners;
 
 public:
-	void AddListener(ListenerType* Listener)
+	void AddListener(BaseNotifyOnChanged* Listener)
 	{
 		if (Listener)
 		{
@@ -28,24 +42,24 @@ public:
 		}
 	}
 
-	void RemoveListener(ListenerType* Listener)
+	void RemoveListener(BaseNotifyOnChanged* Listener)
 	{
 		Listeners.Remove(Listener);
 	}
 
-	void PreChange(const TType* Changed)
+	void PreChange(const TType* Changed, TChangeInfo Info)
 	{
 		for (auto Listener : Listeners)
 		{
-			Listener->PreChange(Changed);
+			Listener->PreChange(Changed, Info);
 		}
 	}
 
-	void PostChange(const TType* Changed)
+	void PostChange(const TType* Changed, TChangeInfo Info)
 	{
 		for (auto Listener : Listeners)
 		{
-			Listener->PostChange(Changed);
+			Listener->PostChange(Changed, Info);
 		}
 	}
 };
\ No newline at end of file
diff --git a/Engine/Source/Editor/UnrealEd/Public/Kismet2/StructureEditorUtils.h b/Engine/Source/Editor/UnrealEd/Public/Kismet2/StructureEditorUtils.h
--- a/Engine/Source/Editor/UnrealEd/Public/Kismet2/StructureEditorUtils.h
+++ b/Engine/Source/Editor/UnrealEd/Public/Kismet2/StructureEditorUtils.h
@@ -7,12 +7,20 @@
 class UNREALED_API FStructureEditorUtils
 {
 public:
+	enum EStructureEditorChangeInfo
+	{
+		Changed,
+	};
 
-	class FStructEditorManager : public FListenerManager<UUserDefinedStruct>
+	class FStructEditorManager : public FListenerManager<UUserDefinedStruct, EStructureEditorChangeInfo>
 	{
-		FStructEditorManager();
+		FStructEditorManager() {}
 	public:
 		UNREALED_API static FStructEditorManager& Get();
+
+		class UNREALED_API ListenerType : public InnerListenerType<FStructEditorManager>
+		{
+		};
 	};
 
 	typedef FStructEditorManager::ListenerType INotifyOnStructChanged;
