DataTableEditor UI: A row can be selected (in Row Editor tab) by clicking on Data Table Editor tab. UI tweaks.

diff --git a/Engine/Source/Editor/DataTableEditor/Private/DataTableEditor.cpp b/Engine/Source/Editor/DataTableEditor/Private/DataTableEditor.cpp
--- a/Engine/Source/Editor/DataTableEditor/Private/DataTableEditor.cpp
+++ b/Engine/Source/Editor/DataTableEditor/Private/DataTableEditor.cpp
@@ -127,6 +127,26 @@ FLinearColor FDataTableEditor::GetWorldCentricTabColorScale() const
 	return FLinearColor( 0.0f, 0.0f, 0.2f, 0.5f );
 }
 
+FSlateColor FDataTableEditor::GetRowColor(FName RowName) const
+{
+	if (RowName == HighlightedRowName)
+	{
+		return FSlateColor(FColorList::Orange);
+	}
+	return FSlateColor::UseForeground();
+}
+
+FReply FDataTableEditor::OnRowClicked(const FGeometry&, const FPointerEvent&, FName RowName)
+{
+	if (HighlightedRowName != RowName)
+	{
+		SetHighlightedRow(RowName);
+		CallbackOnRowHighlighted.ExecuteIfBound(HighlightedRowName);
+	}
+
+	return FReply::Handled();
+}
+
 TSharedPtr<SUniformGridPanel> FDataTableEditor::CreateGridPanel()
 {
 	TSharedPtr<SUniformGridPanel> GridPanel = SNew(SUniformGridPanel).SlotPadding( FMargin( 2.0f ) );
@@ -152,31 +172,38 @@ TSharedPtr<SUniformGridPanel> FDataTableEditor::CreateGridPanel()
 		{
 			if (RowsVisibility[i])
 			{
-				bool bIsHeader =  (i == 0);
+				const bool bIsHeader = (i == 0);
+				const FLinearColor RowColor = (RowIndex % 2 == 0) ? FLinearColor::Gray : FLinearColor::Black;
+
+				TArray<FString>& Row = CachedDataTable[i];
+				FName RowName(*Row[0]);
+				TAttribute<FSlateColor> ForegroundColor = bIsHeader
+					? FSlateColor::UseForeground()
+					: TAttribute<FSlateColor>::Create(
+					TAttribute<FSlateColor>::FGetter::CreateSP(this, &FDataTableEditor::GetRowColor, RowName));
+
+				auto RowClickCallback = bIsHeader
+					? FPointerEventHandler()
+					: FPointerEventHandler::CreateSP(this, &FDataTableEditor::OnRowClicked, RowName);
 
-				FLinearColor RowColor = (RowIndex % 2 == 0) ? FLinearColor::Gray : FLinearColor::Black;				
-				if (bIsHeader)
-				{
-					RowColor = FLinearColor::Gray;
-				}
-				
-				TArray<FString>& Row = CachedDataTable[i];				
 				for(int Column = 0;Column<Row.Num();++Column)
 				{
 					GridPanel->AddSlot(Column, RowIndex)
+					[
+						SNew(SBorder)
+						.Padding(1)
+						.BorderImage(FEditorStyle::GetBrush("ToolPanel.GroupBorder"))
+						.BorderBackgroundColor(RowColor)
+						.ForegroundColor(ForegroundColor)
+						.OnMouseButtonDown(RowClickCallback)
 						[
-							SNew(SBorder)
-							.Padding(1)
-							.BorderImage( FEditorStyle::GetBrush( "ToolPanel.GroupBorder" ) )
-							.BorderBackgroundColor(RowColor)						
-							[
-								SNew(STextBlock)
-								.Text(Row[Column])
-								.ToolTipText(bIsHeader 
-								?	(FString::Printf(TEXT("Column '%s"), *ColumnTitles[Column])) 
-								:	(FString::Printf(TEXT("%s: %s"), *ColumnTitles[Column], *Row[Column])))
-							]
-						];
+							SNew(STextBlock)
+							.Text(Row[Column])
+							.ToolTipText(bIsHeader 
+							?	(FString::Printf(TEXT("Column '%s"), *ColumnTitles[Column])) 
+							:	(FString::Printf(TEXT("%s: %s"), *ColumnTitles[Column], *Row[Column])))
+						]
+					];
 				}
 
 				++RowIndex;
@@ -255,7 +282,10 @@ TSharedRef<SVerticalBox> FDataTableEditor::CreateContentBox()
 
 TSharedRef<SWidget> FDataTableEditor::CreateRowEditorBox()
 {
-	return SNew(SRowEditor, DataTable.Get());
+	auto RowEditor = SNew(SRowEditor, DataTable.Get());
+	RowEditor->RowSelectedCallback.BindSP(this, &FDataTableEditor::SetHighlightedRow);
+	CallbackOnRowHighlighted.BindSP(RowEditor, &SRowEditor::SelectRow);
+	return RowEditor;
 }
 
 TSharedRef<SDockTab> FDataTableEditor::SpawnTab_RowEditor(const FSpawnTabArgs& Args)
@@ -306,5 +336,9 @@ TSharedRef<SDockTab> FDataTableEditor::SpawnTab_DataTable( const FSpawnTabArgs&
 		];
 }
 
+void FDataTableEditor::SetHighlightedRow(FName Name)
+{
+	HighlightedRowName = Name;
+}
 
 #undef LOCTEXT_NAMESPACE
diff --git a/Engine/Source/Editor/DataTableEditor/Private/DataTableEditor.h b/Engine/Source/Editor/DataTableEditor/Private/DataTableEditor.h
--- a/Engine/Source/Editor/DataTableEditor/Private/DataTableEditor.h
+++ b/Engine/Source/Editor/DataTableEditor/Private/DataTableEditor.h
@@ -10,6 +10,8 @@
 
 class UDataTable;
 
+DECLARE_DELEGATE_OneParam(FOnRowHighlighted, FName /*Row name*/);
+
 /** Viewer/editor for a DataTable */
 class FDataTableEditor : public IDataTableEditor
 	, public FStructureEditorUtils::INotifyOnStructChanged
@@ -53,6 +55,8 @@ public:
 	virtual void PreChange(const UDataTable* Changed, FDataTableEditorUtils::EDataTableChangeInfo Info) override;
 	virtual void PostChange(const UDataTable* Changed, FDataTableEditorUtils::EDataTableChangeInfo Info) override;
 
+	void SetHighlightedRow(FName Name);
+
 private:
 
 	void ReloadVisibleData();
@@ -69,6 +73,10 @@ private:
 	/**	Spawns the tab with the Row Editor inside */
 	TSharedRef<SDockTab> SpawnTab_RowEditor(const FSpawnTabArgs& Args);
 
+	FSlateColor GetRowColor(FName RowName) const;
+
+	FReply OnRowClicked(const FGeometry&, const FPointerEvent&, FName RowName);
+
 private:
 
 	/** Cached table data */
@@ -89,6 +97,10 @@ private:
 	/* The DataTable that is active in the editor */
 	TAssetPtr<UDataTable> DataTable;
 
+	FName HighlightedRowName;
+
+	FOnRowHighlighted CallbackOnRowHighlighted;
+
 	/**	The tab id for the data table tab */
 	static const FName DataTableTabId;
 
diff --git a/Engine/Source/Editor/DataTableEditor/Private/SRowEditor.cpp b/Engine/Source/Editor/DataTableEditor/Private/SRowEditor.cpp
--- a/Engine/Source/Editor/DataTableEditor/Private/SRowEditor.cpp
+++ b/Engine/Source/Editor/DataTableEditor/Private/SRowEditor.cpp
@@ -175,6 +175,8 @@ void SRowEditor::Restore()
 	{
 		StructureDetailsView->SetStructureData(CurrentRow);
 	}
+
+	RowSelectedCallback.ExecuteIfBound(FinalName);
 }
 
 UScriptStruct* SRowEditor::GetScriptStruct() const
@@ -217,6 +219,23 @@ void SRowEditor::OnSelectionChanged(TSharedPtr<FName> InItem, ESelectInfo::Type
 	}
 }
 
+void SRowEditor::SelectRow(FName InName)
+{
+	TSharedPtr<FName> NewSelectedName;
+	for (auto Name : CachedRowNames)
+	{
+		if (Name.IsValid() && (*Name == InName))
+		{
+			NewSelectedName = Name;
+		}
+	}
+	if (!NewSelectedName->IsValid())
+	{
+		NewSelectedName = MakeShareable(new FName(InName));
+	}
+	OnSelectionChanged(NewSelectedName, ESelectInfo::Direct);
+}
+
 FReply SRowEditor::OnAddClicked()
 {
 	if (DataTable.IsValid())
@@ -270,7 +289,7 @@ void SRowEditor::Construct(const FArguments& InArgs, UDataTable* Changed)
 
 	RefreshNameList();
 	Restore();
-	
+	const float ButtonWidth = 85.0f;
 	ChildSlot
 	[
 		SNew(SVerticalBox)
@@ -281,44 +300,50 @@ void SRowEditor::Construct(const FArguments& InArgs, UDataTable* Changed)
 			+ SHorizontalBox::Slot()
 			.AutoWidth()
 			.Padding(2)
-			.VAlign(EVerticalAlignment::VAlign_Center)
 			[
-				SNew(STextBlock)
-				.Text(this, &SRowEditor::GetStructureDisplayName)
+				SNew(SBox)
+				.MinDesiredWidth(ButtonWidth)
+				[
+					SNew(SButton)
+					.OnClicked(this, &SRowEditor::OnAddClicked)
+					.HAlign(EHorizontalAlignment::HAlign_Center)
+					.Text(LOCTEXT("AddRow", "Add"))
+				]
 			]
 			+ SHorizontalBox::Slot()
 			.AutoWidth()
 			.Padding(2)
 			[
-				SAssignNew(RowComboBox, SComboBox<TSharedPtr<FName>>)
-				.OptionsSource(&CachedRowNames)
-				.OnSelectionChanged(this, &SRowEditor::OnSelectionChanged)
-				.OnGenerateWidget(this, &SRowEditor::OnGenerateWidget)
-				.Content()
+				SNew(SBox)
+				.MinDesiredWidth(ButtonWidth)
 				[
-					SNew(STextBlock).Text(this, &SRowEditor::GetCurrentNameAsText)
+					SNew(SButton)
+					.OnClicked(this, &SRowEditor::OnRemoveClicked)
+					.HAlign(EHorizontalAlignment::HAlign_Center)
+					.Text(LOCTEXT("RemoveRow", "Remove"))
 				]
 			]
 			+ SHorizontalBox::Slot()
 			.AutoWidth()
 			.Padding(2)
+			.VAlign(EVerticalAlignment::VAlign_Center)
 			[
-				SNew(SButton)
-				.OnClicked(this, &SRowEditor::OnAddClicked)
-				.Content()
+				SNew(SBox)
+				.HAlign(EHorizontalAlignment::HAlign_Right)
 				[
-					SNew(STextBlock).Text(LOCTEXT("AddRow", "Add"))
+					SNew(STextBlock).Text(LOCTEXT("RenameRow", "Rename:"))
 				]
 			]
 			+ SHorizontalBox::Slot()
 			.AutoWidth()
 			.Padding(2)
 			[
-				SNew(SButton)
-				.OnClicked(this, &SRowEditor::OnRemoveClicked)
-				.Content()
+				SNew(SBox)
+				.WidthOverride(2 * ButtonWidth)
 				[
-					SNew(STextBlock).Text(LOCTEXT("RemoveRow", "Remove"))
+					SNew(SEditableTextBox)
+					.Text(this, &SRowEditor::GetCurrentNameAsText)
+					.OnTextCommitted(this, &SRowEditor::OnRowRenamed)
 				]
 			]
 			+ SHorizontalBox::Slot()
@@ -326,15 +351,28 @@ void SRowEditor::Construct(const FArguments& InArgs, UDataTable* Changed)
 			.Padding(2)
 			.VAlign(EVerticalAlignment::VAlign_Center)
 			[
-				SNew(STextBlock).Text(LOCTEXT("RenameRow", "Rename:"))
+				SNew(SBox)
+				.HAlign(EHorizontalAlignment::HAlign_Right)
+				[
+					SNew(STextBlock).Text(LOCTEXT("SelectRow", "Select:"))
+				]
 			]
 			+ SHorizontalBox::Slot()
 			.AutoWidth()
 			.Padding(2)
 			[
-				SNew(SEditableTextBox)
-				.Text(this, &SRowEditor::GetCurrentNameAsText)
-				.OnTextCommitted(this, &SRowEditor::OnRowRenamed)
+				SNew(SBox)
+				.WidthOverride(2 * ButtonWidth)
+				[
+					SAssignNew(RowComboBox, SComboBox<TSharedPtr<FName>>)
+					.OptionsSource(&CachedRowNames)
+					.OnSelectionChanged(this, &SRowEditor::OnSelectionChanged)
+					.OnGenerateWidget(this, &SRowEditor::OnGenerateWidget)
+					.Content()
+					[
+						SNew(STextBlock).Text(this, &SRowEditor::GetCurrentNameAsText)
+					]
+				]
 			]
 		]
 		+ SVerticalBox::Slot()
@@ -343,4 +381,5 @@ void SRowEditor::Construct(const FArguments& InArgs, UDataTable* Changed)
 			StructureDetailsView->GetWidget().ToSharedRef()
 		]
 	];
-}
\ No newline at end of file
+}
+
diff --git a/Engine/Source/Editor/DataTableEditor/Private/SRowEditor.h b/Engine/Source/Editor/DataTableEditor/Private/SRowEditor.h
--- a/Engine/Source/Editor/DataTableEditor/Private/SRowEditor.h
+++ b/Engine/Source/Editor/DataTableEditor/Private/SRowEditor.h
@@ -6,6 +6,7 @@
 #include "DataTableEditorUtils.h"
 
 DECLARE_DELEGATE_OneParam(FOnRowModified, FName /*Row name*/);
+DECLARE_DELEGATE_OneParam(FOnRowSelected, FName /*Row name*/);
 
 class SRowEditor : public SCompoundWidget
 	, public FStructureEditorUtils::INotifyOnStructChanged
@@ -26,6 +27,8 @@ public:
 	virtual void PreChange(const UDataTable* Changed, FDataTableEditorUtils::EDataTableChangeInfo Info) override;
 	virtual void PostChange(const UDataTable* Changed, FDataTableEditorUtils::EDataTableChangeInfo Info) override;
 
+	FOnRowSelected RowSelectedCallback;
+
 private:
 
 	TArray<TSharedPtr<FName>> CachedRowNames;
@@ -56,4 +59,6 @@ private:
 public:
 
 	void Construct(const FArguments& InArgs, UDataTable* Changed);
+
+	void SelectRow(FName Name);
 };
\ No newline at end of file
