DataTable RowEditor: When you changed name of row, and selected new row without hitting enter, name is changed on the newly selected row, instead of the old one, which works counter intuitive.

#688 based on pull request

diff --git a/Engine/Source/Editor/DataTableEditor/Private/SRowEditor.cpp b/Engine/Source/Editor/DataTableEditor/Private/SRowEditor.cpp
--- a/Engine/Source/Editor/DataTableEditor/Private/SRowEditor.cpp
+++ b/Engine/Source/Editor/DataTableEditor/Private/SRowEditor.cpp
@@ -216,6 +216,12 @@ void SRowEditor::OnSelectionChanged(TSharedPtr<FName> InItem, ESelectInfo::Type
 		SelectedName = InItem;
 
 		Restore();
+
+		if (RenameTextBox.IsValid())
+		{
+			// refresh name, in case of a pending rename action
+			RenameTextBox->SetText(TAttribute<FText>(this, &SRowEditor::GetCurrentNameAsText));
+		}
 	}
 }
 
@@ -266,8 +272,17 @@ void SRowEditor::OnRowRenamed(const FText& Text, ETextCommit::Type CommitType)
 {
 	if (!GetCurrentNameAsText().EqualTo(Text) && DataTable.IsValid())
 	{
-		const FName OldName = GetCurrentName();
 		const FName NewName = DataTable->MakeValidName(Text.ToString());
+		for (auto Name : CachedRowNames)
+		{
+			if (Name.IsValid() && (*Name == NewName))
+			{
+				 //the name already exists
+				return;
+			}
+		}
+
+		const FName OldName = GetCurrentName();
 		SelectedName = MakeShareable(new FName(NewName));
 		FDataTableEditorUtils::RenameRow(DataTable.Get(), OldName, NewName);
 	}
@@ -341,7 +356,7 @@ void SRowEditor::Construct(const FArguments& InArgs, UDataTable* Changed)
 				SNew(SBox)
 				.WidthOverride(2 * ButtonWidth)
 				[
-					SNew(SEditableTextBox)
+					SAssignNew(RenameTextBox, SEditableTextBox)
 					.Text(this, &SRowEditor::GetCurrentNameAsText)
 					.OnTextCommitted(this, &SRowEditor::OnRowRenamed)
 				]
diff --git a/Engine/Source/Editor/DataTableEditor/Private/SRowEditor.h b/Engine/Source/Editor/DataTableEditor/Private/SRowEditor.h
--- a/Engine/Source/Editor/DataTableEditor/Private/SRowEditor.h
+++ b/Engine/Source/Editor/DataTableEditor/Private/SRowEditor.h
@@ -37,6 +37,7 @@ private:
 	TSharedPtr<class IStructureDetailsView> StructureDetailsView;
 	TSharedPtr<FName> SelectedName;
 	TSharedPtr<SComboBox<TSharedPtr<FName>>> RowComboBox;
+	TSharedPtr<SEditableTextBox> RenameTextBox;
 
 	void RefreshNameList();
 	void CleanBeforeChange();
