Data table UI improvements

UETOOL-291 - Improve the data table UI

The column sizes in the data table can now be resized as desired, and also default to being sized to match their content rather than as a simple proportion of their parent size. Any resized columns are remembered between editing sessions.

The data table view now scrollable horizontally as well as vertically to account for the fact that the data may now be wider than the current view.

diff --git a/Engine/Source/Editor/DataTableEditor/DataTableEditor.Build.cs b/Engine/Source/Editor/DataTableEditor/DataTableEditor.Build.cs
--- a/Engine/Source/Editor/DataTableEditor/DataTableEditor.Build.cs
+++ b/Engine/Source/Editor/DataTableEditor/DataTableEditor.Build.cs
@@ -18,7 +18,8 @@ public class DataTableEditor : ModuleRules
 				"SlateCore",
                 "EditorStyle",
                 "PropertyEditor",
-				"UnrealEd"
+				"UnrealEd",
+				"Json"
 			}
 			);
 
diff --git a/Engine/Source/Editor/DataTableEditor/Private/DataTableEditor.cpp b/Engine/Source/Editor/DataTableEditor/Private/DataTableEditor.cpp
--- a/Engine/Source/Editor/DataTableEditor/Private/DataTableEditor.cpp
+++ b/Engine/Source/Editor/DataTableEditor/Private/DataTableEditor.cpp
@@ -9,6 +9,7 @@
 #include "SDockTab.h"
 #include "SRowEditor.h"
 #include "Engine/DataTable.h"
+#include "Json.h"
  
 #define LOCTEXT_NAMESPACE "DataTableEditor"
 
@@ -40,6 +41,10 @@ FDataTableEditor::FDataTableEditor()
 
 FDataTableEditor::~FDataTableEditor()
 {
+	if (DataTable.IsValid())
+	{
+		SaveLayoutData();
+	}
 }
 
 void FDataTableEditor::PreChange(const class UUserDefinedStruct* Struct, FStructureEditorUtils::EStructureEditorChangeInfo Info)
@@ -147,35 +152,135 @@ FReply FDataTableEditor::OnRowClicked(const FGeometry&, const FPointerEvent&, FN
 	return FReply::Handled();
 }
 
-TSharedPtr<SUniformGridPanel> FDataTableEditor::CreateGridPanel()
+float FDataTableEditor::GetColumnWidth(const int32 ColumnIndex)
+{
+	if (ColumnWidths.IsValidIndex(ColumnIndex))
+	{
+		return ColumnWidths[ColumnIndex];
+	}
+	return 0.0f;
+}
+
+void FDataTableEditor::OnColumnResized(const float NewWidth, const int32 ColumnIndex)
+{
+	if (ColumnWidths.IsValidIndex(ColumnIndex))
+	{
+		ColumnWidths[ColumnIndex] = NewWidth;
+
+		// Update the persistent column widths in the layout data
+		{
+			if (!LayoutData.IsValid())
+			{
+				LayoutData = MakeShareable(new FJsonObject());
+			}
+
+			TSharedPtr<FJsonObject> LayoutColumnWidths;
+			if (!LayoutData->HasField(TEXT("ColumnWidths")))
+			{
+				LayoutColumnWidths = MakeShareable(new FJsonObject());
+				LayoutData->SetObjectField(TEXT("ColumnWidths"), LayoutColumnWidths);
+			}
+			else
+			{
+				LayoutColumnWidths = LayoutData->GetObjectField(TEXT("ColumnWidths"));
+			}
+
+			const FString& ColumnName = CachedRawColumnNames[ColumnIndex];
+			LayoutColumnWidths->SetNumberField(ColumnName, NewWidth);
+		}
+	}
+}
+
+void FDataTableEditor::LoadLayoutData()
+{
+	LayoutData.Reset();
+
+	if (!DataTable.IsValid())
+	{
+		return;
+	}
+
+	const FString LayoutDataFilename = FPaths::GameSavedDir() / TEXT("AssetData") / TEXT("DataTableEditorLayout") / DataTable->GetName() + TEXT(".json");
+
+	FString JsonText;
+	if (FFileHelper::LoadFileToString(JsonText, *LayoutDataFilename))
+	{
+		TSharedRef< TJsonReader<TCHAR> > JsonReader = TJsonReaderFactory<TCHAR>::Create(JsonText);
+		FJsonSerializer::Deserialize(JsonReader, LayoutData);
+	}
+}
+
+void FDataTableEditor::SaveLayoutData()
 {
-	TSharedPtr<SUniformGridPanel> GridPanel = SNew(SUniformGridPanel).SlotPadding( FMargin( 2.0f ) );
+	if (!DataTable.IsValid() || !LayoutData.IsValid())
+	{
+		return;
+	}
+
+	const FString LayoutDataFilename = FPaths::GameSavedDir() / TEXT("AssetData") / TEXT("DataTableEditorLayout") / DataTable->GetName() + TEXT(".json");
+
+	FString JsonText;
+	TSharedRef< TJsonWriter< TCHAR, TPrettyJsonPrintPolicy<TCHAR> > > JsonWriter = TJsonWriterFactory< TCHAR, TPrettyJsonPrintPolicy<TCHAR> >::Create(&JsonText);
+	if (FJsonSerializer::Serialize(LayoutData.ToSharedRef(), JsonWriter))
+	{
+		FFileHelper::SaveStringToFile(JsonText, *LayoutDataFilename);
+	}
+}
+
+TSharedRef<SWidget> FDataTableEditor::CreateGridPanel()
+{
+	TSharedRef<SVerticalBox> VerticalBox = SNew(SVerticalBox);
 
 	if (DataTable.IsValid())
 	{
 		if (CachedDataTable.Num() == 0)
 		{
 			CachedDataTable = DataTable->GetTableData();
+			check(CachedDataTable.Num() > 0); // There should always be at least the column titles row
+
+			CachedRawColumnNames = DataTable->GetUniqueColumnTitles();
 
 			RowsVisibility.SetNum(CachedDataTable.Num());
-			for (int32 i = 0; i < RowsVisibility.Num(); ++i)
+			for (bool& RowVisibility : RowsVisibility)
+			{
+				RowVisibility = true;
+			}
+
+			ColumnWidths.SetNumZeroed(CachedRawColumnNames.Num());
+			
+			// Load the persistent column widths from the layout data
 			{
-				RowsVisibility[i] = true;
+				const TSharedPtr<FJsonObject>* LayoutColumnWidths = nullptr;
+				if (LayoutData.IsValid() && LayoutData->TryGetObjectField(TEXT("ColumnWidths"), LayoutColumnWidths))
+				{
+					for(int32 ColumnIndex = 0; ColumnIndex < CachedRawColumnNames.Num(); ++ColumnIndex)
+					{
+						const FString& ColumnName = CachedRawColumnNames[ColumnIndex];
+
+						double LayoutColumnWidth = 0.0f;
+						if ((*LayoutColumnWidths)->TryGetNumberField(ColumnName, LayoutColumnWidth))
+						{
+							ColumnWidths[ColumnIndex] = static_cast<float>(LayoutColumnWidth);
+						}
+					}
+				}
 			}
 		}
 
 		check(CachedDataTable.Num() > 0 && CachedDataTable.Num() == RowsVisibility.Num());
+		check(CachedDataTable.Num() > 0 && CachedDataTable[0].Num() == ColumnWidths.Num());
+		check(CachedRawColumnNames.Num() > 0 && CachedRawColumnNames.Num() == ColumnWidths.Num());
 
-		int32 RowIndex = 0;
+		int32 VisibleRowIndex = 0;
 		TArray<FString>& ColumnTitles = CachedDataTable[0];
-		for(int i = 0;i<CachedDataTable.Num();++i)
+		for(int32 RowIndex = 0; RowIndex < CachedDataTable.Num(); ++RowIndex)
 		{
-			if (RowsVisibility[i])
+			if (RowsVisibility[RowIndex])
 			{
-				const bool bIsHeader = (i == 0);
-				const FLinearColor RowColor = (RowIndex % 2 == 0) ? FLinearColor::Gray : FLinearColor::Black;
+				const bool bIsHeader = (RowIndex == 0);
+				const FLinearColor RowColor = (VisibleRowIndex % 2 == 0) ? FLinearColor::Gray : FLinearColor::Black;
 
-				TArray<FString>& Row = CachedDataTable[i];
+				TArray<FString>& Row = CachedDataTable[RowIndex];
 				FName RowName(*Row[0]);
 				TAttribute<FSlateColor> ForegroundColor = bIsHeader
 					? FSlateColor::UseForeground()
@@ -186,31 +291,59 @@ TSharedPtr<SUniformGridPanel> FDataTableEditor::CreateGridPanel()
 					? FPointerEventHandler()
 					: FPointerEventHandler::CreateSP(this, &FDataTableEditor::OnRowClicked, RowName);
 
-				for(int Column = 0;Column<Row.Num();++Column)
+				const float SplitterHandleSize = 2.0f;
+				TSharedRef<SSplitter> RowSplitter = SNew(SSplitter)
+					.PhysicalSplitterHandleSize(SplitterHandleSize)
+					.HitDetectionSplitterHandleSize(4.0f);
+
+				VerticalBox->AddSlot()
+				[
+					RowSplitter
+				];
+
+				for(int32 ColumnIndex = 0; ColumnIndex < Row.Num(); ++ColumnIndex)
 				{
-					GridPanel->AddSlot(Column, RowIndex)
+					TSharedPtr<SWidget> ColumnEntryWidget;
+
+					RowSplitter->AddSlot()
+					.Value(TAttribute<float>::Create(TAttribute<float>::FGetter::CreateSP(this, &FDataTableEditor::GetColumnWidth, ColumnIndex)))
+					.OnSlotResized(SSplitter::FOnSlotResized::CreateSP(this, &FDataTableEditor::OnColumnResized, ColumnIndex))
+					.SizeRule(SSplitter::AbsoluteSize)
 					[
-						SNew(SBorder)
-						.Padding(1)
+						SAssignNew(ColumnEntryWidget, SBorder)
+						.Padding(2.0f)
 						.BorderImage(FEditorStyle::GetBrush("ToolPanel.GroupBorder"))
 						.BorderBackgroundColor(RowColor)
 						.ForegroundColor(ForegroundColor)
 						.OnMouseButtonDown(RowClickCallback)
 						[
 							SNew(STextBlock)
-							.Text(FText::FromString(Row[Column]))
+							.Text(FText::FromString(Row[ColumnIndex]))
 							.ToolTipText(bIsHeader 
-							?	(FText::Format(LOCTEXT("ColumnHeaderNameFmt", "Column '{0}'"), FText::FromString(ColumnTitles[Column]))) 
-							:	(FText::Format(LOCTEXT("ColumnRowNameFmt", "{0}: {1}"), FText::FromString(ColumnTitles[Column]), FText::FromString(Row[Column]))))
+							?	(FText::Format(LOCTEXT("ColumnHeaderNameFmt", "Column '{0}'"), FText::FromString(ColumnTitles[ColumnIndex]))) 
+							:	(FText::Format(LOCTEXT("ColumnRowNameFmt", "{0}: {1}"), FText::FromString(ColumnTitles[ColumnIndex]), FText::FromString(Row[ColumnIndex]))))
 						]
 					];
+
+					ColumnEntryWidget->SlatePrepass(1.0f);
+					ColumnWidths[ColumnIndex] = FMath::Max(ColumnWidths[ColumnIndex], ColumnEntryWidget->GetDesiredSize().X + SplitterHandleSize);
 				}
 
-				++RowIndex;
+				// Dummy splitter slot to allow the last column to be resized
+				RowSplitter->AddSlot()
+				.Value(TAttribute<float>::Create(TAttribute<float>::FGetter::CreateSP(this, &FDataTableEditor::GetColumnWidth, static_cast<int32>(INDEX_NONE))))
+				.OnSlotResized(SSplitter::FOnSlotResized::CreateSP(this, &FDataTableEditor::OnColumnResized, static_cast<int32>(INDEX_NONE)))
+				.SizeRule(SSplitter::AbsoluteSize)
+				[
+					SNullWidget::NullWidget
+				];
+
+				++VisibleRowIndex;
 			}
 		}
 	}
-	return GridPanel;
+
+	return VerticalBox;
 }
 
 void FDataTableEditor::OnSearchTextChanged(const FText& SearchText)
@@ -256,13 +389,21 @@ void FDataTableEditor::ReloadVisibleData()
 		ScrollBoxWidget->ClearChildren();
 		ScrollBoxWidget->AddSlot()
 			[
-				CreateGridPanel().ToSharedRef()
+				CreateGridPanel()
 			];
 	}
 }
 
 TSharedRef<SVerticalBox> FDataTableEditor::CreateContentBox()
 {
+	TSharedRef<SScrollBar> HorizontalScrollBar = SNew(SScrollBar)
+		.Orientation(Orient_Horizontal)
+		.Thickness(FVector2D(5, 5));
+
+	TSharedRef<SScrollBar> VerticalScrollBar = SNew(SScrollBar)
+		.Orientation(Orient_Vertical)
+		.Thickness(FVector2D(5, 5));
+
 	return SNew(SVerticalBox)
 		+SVerticalBox::Slot()
 		.AutoHeight()
@@ -272,11 +413,34 @@ TSharedRef<SVerticalBox> FDataTableEditor::CreateContentBox()
 		]
 		+SVerticalBox::Slot()
 		[
-			SAssignNew(ScrollBoxWidget, SScrollBox)
-			+SScrollBox::Slot()
+			SNew(SHorizontalBox)
+			+SHorizontalBox::Slot()
 			[
-				CreateGridPanel().ToSharedRef()
+				SNew(SScrollBox)
+				.Orientation(Orient_Horizontal)
+				.ExternalScrollbar(HorizontalScrollBar)
+				+SScrollBox::Slot()
+				[
+					SAssignNew(ScrollBoxWidget, SScrollBox)
+					.Orientation(Orient_Vertical)
+					.ExternalScrollbar(VerticalScrollBar)
+					.ConsumeMouseWheel(EConsumeMouseWheel::Always) // Always consume the mouse wheel events to prevent the outer scroll box from scrolling
+					+SScrollBox::Slot()
+					[
+						CreateGridPanel()
+					]
+				]
+			]
+			+SHorizontalBox::Slot()
+			.AutoWidth()
+			[
+				VerticalScrollBar
 			]
+		]
+		+SVerticalBox::Slot()
+		.AutoHeight()
+		[
+			HorizontalScrollBar
 		];
 }
 
@@ -315,13 +479,13 @@ TSharedRef<SDockTab> FDataTableEditor::SpawnTab_DataTable( const FSpawnTabArgs&
 
 	DataTable = Cast<UDataTable>(GetEditingObject());
 
+	LoadLayoutData();
+
 	TSharedRef<SVerticalBox> ContentBox = CreateContentBox();
 
 	GridPanelOwner = 
 		SNew(SBorder)
 		.Padding(2)
-		.VAlign(VAlign_Top)
-		.HAlign(HAlign_Left)
 		.BorderImage( FEditorStyle::GetBrush( "ToolPanel.GroupBorder" ) )
 		[
 			ContentBox
diff --git a/Engine/Source/Editor/DataTableEditor/Private/DataTableEditor.h b/Engine/Source/Editor/DataTableEditor/Private/DataTableEditor.h
--- a/Engine/Source/Editor/DataTableEditor/Private/DataTableEditor.h
+++ b/Engine/Source/Editor/DataTableEditor/Private/DataTableEditor.h
@@ -9,6 +9,7 @@
 #include "DataTableEditorUtils.h"
 
 class UDataTable;
+class FJsonObject;
 
 DECLARE_DELEGATE_OneParam(FOnRowHighlighted, FName /*Row name*/);
 
@@ -38,8 +39,8 @@ public:
 	/** Destructor */
 	virtual ~FDataTableEditor();
 
-	/* Create the uniform grid panel */
-	TSharedPtr<SUniformGridPanel> CreateGridPanel();
+	/* Create the grid panel */
+	TSharedRef<SWidget> CreateGridPanel();
 
 	/** IToolkit interface */
 	virtual FName GetToolkitFName() const override;
@@ -77,14 +78,28 @@ private:
 
 	FReply OnRowClicked(const FGeometry&, const FPointerEvent&, FName RowName);
 
+	float GetColumnWidth(const int32 ColumnIndex);
+
+	void OnColumnResized(const float NewWidth, const int32 ColumnIndex);
+
+	void LoadLayoutData();
+
+	void SaveLayoutData();
+
 private:
 
 	/** Cached table data */
 	TArray<TArray<FString> > CachedDataTable;
 
+	/** Cached raw column names */
+	TArray<FString> CachedRawColumnNames;
+
 	/** Visibility of data table rows */
 	TArray<bool> RowsVisibility;
 
+	/** Widths of data table columns */
+	TArray<float> ColumnWidths;
+
 	/** Search box */
 	TSharedPtr<SWidget> SearchBox;
 
@@ -97,6 +112,9 @@ private:
 	/* The DataTable that is active in the editor */
 	TAssetPtr<UDataTable> DataTable;
 
+	/** The layout data for the currently loaded data table */
+	TSharedPtr<FJsonObject> LayoutData;
+
 	FName HighlightedRowName;
 
 	FOnRowHighlighted CallbackOnRowHighlighted;
diff --git a/Engine/Source/Runtime/Engine/Classes/Engine/DataTable.h b/Engine/Source/Runtime/Engine/Classes/Engine/DataTable.h
--- a/Engine/Source/Runtime/Engine/Classes/Engine/DataTable.h
+++ b/Engine/Source/Runtime/Engine/Classes/Engine/DataTable.h
@@ -173,9 +173,12 @@ public:
 	*/
 	ENGINE_API TArray<FString> CreateTableFromJSONString(const FString& InString);
 
-	/** Get array of all the column titles */
+	/** Get an array of all the column titles, using the friendly display name from the property */
 	ENGINE_API TArray<FString> GetColumnTitles() const;
 
+	/** Get an array of all the column titles, using the unique name from the property */
+	ENGINE_API TArray<FString> GetUniqueColumnTitles() const;
+
 	TArray<UProperty*> GetTablePropertyArray(const TArray<const TCHAR*>& Cells, UStruct* RowStruct, TArray<FString>& OutProblems);
 
 	/** Get array for each row in the table. The first row is the titles*/
diff --git a/Engine/Source/Runtime/Engine/Private/DataTable.cpp b/Engine/Source/Runtime/Engine/Private/DataTable.cpp
--- a/Engine/Source/Runtime/Engine/Private/DataTable.cpp
+++ b/Engine/Source/Runtime/Engine/Private/DataTable.cpp
@@ -419,6 +419,20 @@ TArray<FString> UDataTable::GetColumnTitles() const
 	return Result;
 }
 
+TArray<FString> UDataTable::GetUniqueColumnTitles() const
+{
+	TArray<FString> Result;
+	Result.Add(TEXT("Name"));
+	for (TFieldIterator<UProperty> It(RowStruct); It; ++It)
+	{
+		UProperty* Prop = *It;
+		check(Prop != NULL);
+		const FString DisplayName = Prop->GetName();
+		Result.Add(DisplayName);
+	}
+	return Result;
+}
+
 TArray< TArray<FString> > UDataTable::GetTableData() const
 {
 	 TArray< TArray<FString> > Result;
