Stopped data table columns from auto-sizing even after the user has explicitly sized them

diff --git a/Engine/Source/Editor/DataTableEditor/Private/DataTableEditor.cpp b/Engine/Source/Editor/DataTableEditor/Private/DataTableEditor.cpp
--- a/Engine/Source/Editor/DataTableEditor/Private/DataTableEditor.cpp
+++ b/Engine/Source/Editor/DataTableEditor/Private/DataTableEditor.cpp
@@ -156,7 +156,7 @@ float FDataTableEditor::GetColumnWidth(const int32 ColumnIndex)
 {
 	if (ColumnWidths.IsValidIndex(ColumnIndex))
 	{
-		return ColumnWidths[ColumnIndex];
+		return ColumnWidths[ColumnIndex].CurrentWidth;
 	}
 	return 0.0f;
 }
@@ -165,7 +165,9 @@ void FDataTableEditor::OnColumnResized(const float NewWidth, const int32 ColumnI
 {
 	if (ColumnWidths.IsValidIndex(ColumnIndex))
 	{
-		ColumnWidths[ColumnIndex] = NewWidth;
+		FColumnWidth& ColumnWidth = ColumnWidths[ColumnIndex];
+		ColumnWidth.bIsAutoSized = false;
+		ColumnWidth.CurrentWidth = NewWidth;
 
 		// Update the persistent column widths in the layout data
 		{
@@ -246,7 +248,7 @@ TSharedRef<SWidget> FDataTableEditor::CreateGridPanel()
 				RowVisibility = true;
 			}
 
-			ColumnWidths.SetNumZeroed(CachedRawColumnNames.Num());
+			ColumnWidths.SetNum(CachedRawColumnNames.Num());
 			
 			// Load the persistent column widths from the layout data
 			{
@@ -260,7 +262,9 @@ TSharedRef<SWidget> FDataTableEditor::CreateGridPanel()
 						double LayoutColumnWidth = 0.0f;
 						if ((*LayoutColumnWidths)->TryGetNumberField(ColumnName, LayoutColumnWidth))
 						{
-							ColumnWidths[ColumnIndex] = static_cast<float>(LayoutColumnWidth);
+							FColumnWidth& ColumnWidth = ColumnWidths[ColumnIndex];
+							ColumnWidth.bIsAutoSized = false;
+							ColumnWidth.CurrentWidth = static_cast<float>(LayoutColumnWidth);
 						}
 					}
 				}
@@ -301,6 +305,7 @@ TSharedRef<SWidget> FDataTableEditor::CreateGridPanel()
 					RowSplitter
 				];
 
+				float ColumnDesiredWidth = 0.0f;
 				for(int32 ColumnIndex = 0; ColumnIndex < Row.Num(); ++ColumnIndex)
 				{
 					TSharedPtr<SWidget> ColumnEntryWidget;
@@ -325,8 +330,12 @@ TSharedRef<SWidget> FDataTableEditor::CreateGridPanel()
 						]
 					];
 
-					ColumnEntryWidget->SlatePrepass(1.0f);
-					ColumnWidths[ColumnIndex] = FMath::Max(ColumnWidths[ColumnIndex], ColumnEntryWidget->GetDesiredSize().X + SplitterHandleSize);
+					FColumnWidth& ColumnWidth = ColumnWidths[ColumnIndex];
+					if (ColumnWidth.bIsAutoSized)
+					{
+						ColumnEntryWidget->SlatePrepass(1.0f);
+						ColumnWidth.CurrentWidth = FMath::Max(ColumnWidth.CurrentWidth, ColumnEntryWidget->GetDesiredSize().X + SplitterHandleSize);
+					}
 				}
 
 				// Dummy splitter slot to allow the last column to be resized
diff --git a/Engine/Source/Editor/DataTableEditor/Private/DataTableEditor.h b/Engine/Source/Editor/DataTableEditor/Private/DataTableEditor.h
--- a/Engine/Source/Editor/DataTableEditor/Private/DataTableEditor.h
+++ b/Engine/Source/Editor/DataTableEditor/Private/DataTableEditor.h
@@ -88,6 +88,21 @@ private:
 
 private:
 
+	/** Struct holding information about the current column widths */
+	struct FColumnWidth
+	{
+		FColumnWidth()
+			: bIsAutoSized(true)
+			, CurrentWidth(0.0f)
+		{
+		}
+
+		/** True if this column is being auto-sized rather than sized by the user */
+		bool bIsAutoSized;
+		/** The width of the column, either sized by the user, or auto-sized */
+		float CurrentWidth;
+	};
+
 	/** Cached table data */
 	TArray<TArray<FString> > CachedDataTable;
 
@@ -98,7 +113,7 @@ private:
 	TArray<bool> RowsVisibility;
 
 	/** Widths of data table columns */
-	TArray<float> ColumnWidths;
+	TArray<FColumnWidth> ColumnWidths;
 
 	/** Search box */
 	TSharedPtr<SWidget> SearchBox;
