Added the ability for data tables to edit class and object properties

UETOOL-290 - Data tables don't let you edit asset/class references in the row editor panel

Also fixed the fact that the properties panel wouldn't scroll.

diff --git a/Engine/Source/Editor/DataTableEditor/Private/SRowEditor.cpp b/Engine/Source/Editor/DataTableEditor/Private/SRowEditor.cpp
--- a/Engine/Source/Editor/DataTableEditor/Private/SRowEditor.cpp
+++ b/Engine/Source/Editor/DataTableEditor/Private/SRowEditor.cpp
@@ -299,7 +299,7 @@ void SRowEditor::Construct(const FArguments& InArgs, UDataTable* Changed)
 		ViewArgs.bHideSelectionTip = false;
 		ViewArgs.bShowActorLabel = false;
 
-		StructureDetailsView = PropertyModule.CreateStructureDetailView(ViewArgs, CurrentRow, false, LOCTEXT("RowValue", "Row Value"));
+		StructureDetailsView = PropertyModule.CreateStructureDetailView(ViewArgs, CurrentRow, true/*bShowObjects*/, LOCTEXT("RowValue", "Row Value"));
 		StructureDetailsView->GetOnFinishedChangingPropertiesDelegate().AddSP(this, &SRowEditor::OnFinishedChangingProperties);
 	}
 
@@ -392,7 +392,6 @@ void SRowEditor::Construct(const FArguments& InArgs, UDataTable* Changed)
 			]
 		]
 		+ SVerticalBox::Slot()
-		.AutoHeight()
 		[
 			StructureDetailsView->GetWidget().ToSharedRef()
 		]
