Added the ability to re-order rows in the data table editor

UETOOL-295 - Can't re-order rows in a data table

Added an option to move a row up or down, or to the top or bottom of the data table.

Updated the row editor UI to be more consistent with other editors, and re-arranged some of the controls on the toolbar, adding tooltips where appropriate.

Fixed the highlighted row in the data table not matching the selected row editor row when a data table was first opened.

diff --git a/Engine/Source/Editor/DataTableEditor/Private/DataTableEditor.cpp b/Engine/Source/Editor/DataTableEditor/Private/DataTableEditor.cpp
--- a/Engine/Source/Editor/DataTableEditor/Private/DataTableEditor.cpp
+++ b/Engine/Source/Editor/DataTableEditor/Private/DataTableEditor.cpp
@@ -350,6 +350,21 @@ TSharedRef<SWidget> FDataTableEditor::CreateGridPanel()
 				++VisibleRowIndex;
 			}
 		}
+
+		// Update the currently selected row to try and ensure it's valid
+		if (HighlightedRowName.IsNone() || !DataTable->RowMap.Contains(HighlightedRowName))
+		{
+			if (CachedDataTable.Num() > 2)
+			{
+				// 1 is the first row in the table
+				// 0 is the name column
+				SetHighlightedRow(*CachedDataTable[1][0]);
+			}
+			else
+			{
+				SetHighlightedRow(NAME_None);
+			}
+		}
 	}
 
 	return VerticalBox;
diff --git a/Engine/Source/Editor/DataTableEditor/Private/SRowEditor.cpp b/Engine/Source/Editor/DataTableEditor/Private/SRowEditor.cpp
--- a/Engine/Source/Editor/DataTableEditor/Private/SRowEditor.cpp
+++ b/Engine/Source/Editor/DataTableEditor/Private/SRowEditor.cpp
@@ -269,6 +269,27 @@ FReply SRowEditor::OnRemoveClicked()
 	return FReply::Handled();
 }
 
+FReply SRowEditor::OnMoveRowClicked(FDataTableEditorUtils::ERowMoveDirection MoveDirection)
+{
+	if (DataTable.IsValid())
+	{
+		const FName RowToMove = GetCurrentName();
+		FDataTableEditorUtils::MoveRow(DataTable.Get(), RowToMove, MoveDirection);
+	}
+	return FReply::Handled();
+}
+
+FReply SRowEditor::OnMoveToExtentClicked(FDataTableEditorUtils::ERowMoveDirection MoveDirection)
+{
+	if (DataTable.IsValid())
+	{
+		// We move by the row map size, as FDataTableEditorUtils::MoveRow will automatically clamp this as appropriate
+		const FName RowToMove = GetCurrentName();
+		FDataTableEditorUtils::MoveRow(DataTable.Get(), RowToMove, MoveDirection, DataTable->RowMap.Num());
+	}
+	return FReply::Handled();
+}
+
 void SRowEditor::OnRowRenamed(const FText& Text, ETextCommit::Type CommitType)
 {
 	if (!GetCurrentNameAsText().EqualTo(Text) && DataTable.IsValid())
@@ -317,13 +338,32 @@ void SRowEditor::Construct(const FArguments& InArgs, UDataTable* Changed)
 			.AutoWidth()
 			.Padding(2)
 			[
-				SNew(SBox)
-				.MinDesiredWidth(ButtonWidth)
+				SNew(SButton)
+				.ButtonStyle(FEditorStyle::Get(), "HoverHintOnly")
+				.ForegroundColor(FSlateColor::UseForeground())
+				.HAlign(HAlign_Center)
+				.VAlign(VAlign_Center)
+				.OnClicked(this, &SRowEditor::OnAddClicked)
+				.ToolTipText(LOCTEXT("AddRowTooltip", "Add a new row to the data table"))
 				[
-					SNew(SButton)
-					.OnClicked(this, &SRowEditor::OnAddClicked)
-					.HAlign(EHorizontalAlignment::HAlign_Center)
-					.Text(LOCTEXT("AddRow", "Add"))
+					SNew(SImage)
+					.Image(FEditorStyle::Get().GetBrush("Plus"))
+				]
+			]
+			+ SHorizontalBox::Slot()
+			.AutoWidth()
+			.Padding(2)
+			[
+				SNew(SButton)
+				.ButtonStyle(FEditorStyle::Get(), "HoverHintOnly")
+				.ForegroundColor(FSlateColor::UseForeground())
+				.HAlign(HAlign_Center)
+				.VAlign(VAlign_Center)
+				.OnClicked(this, &SRowEditor::OnRemoveClicked)
+				.ToolTipText(LOCTEXT("RemoveRowTooltip", "Remove the currently selected row from the data table"))
+				[
+					SNew(SImage)
+					.Image(FEditorStyle::Get().GetBrush("Cross"))
 				]
 			]
 			+ SHorizontalBox::Slot()
@@ -331,23 +371,32 @@ void SRowEditor::Construct(const FArguments& InArgs, UDataTable* Changed)
 			.Padding(2)
 			[
 				SNew(SBox)
-				.MinDesiredWidth(ButtonWidth)
+				.WidthOverride(2 * ButtonWidth)
+				.ToolTipText(LOCTEXT("SelectedRowTooltip", "Select a row to edit"))
 				[
-					SNew(SButton)
-					.OnClicked(this, &SRowEditor::OnRemoveClicked)
-					.HAlign(EHorizontalAlignment::HAlign_Center)
-					.Text(LOCTEXT("RemoveRow", "Remove"))
+					SAssignNew(RowComboBox, SComboBox<TSharedPtr<FName>>)
+					.OptionsSource(&CachedRowNames)
+					.OnSelectionChanged(this, &SRowEditor::OnSelectionChanged)
+					.OnGenerateWidget(this, &SRowEditor::OnGenerateWidget)
+					.Content()
+					[
+						SNew(STextBlock).Text(this, &SRowEditor::GetCurrentNameAsText)
+					]
 				]
 			]
 			+ SHorizontalBox::Slot()
+			[
+				SNew(SSpacer)
+			]
+			+ SHorizontalBox::Slot()
 			.AutoWidth()
 			.Padding(2)
 			.VAlign(EVerticalAlignment::VAlign_Center)
 			[
 				SNew(SBox)
-				.HAlign(EHorizontalAlignment::HAlign_Right)
+				.HAlign(HAlign_Right)
 				[
-					SNew(STextBlock).Text(LOCTEXT("RenameRow", "Rename:"))
+					SNew(STextBlock).Text(LOCTEXT("RowNameLabel", "Row Name:"))
 				]
 			]
 			+ SHorizontalBox::Slot()
@@ -365,29 +414,69 @@ void SRowEditor::Construct(const FArguments& InArgs, UDataTable* Changed)
 			+ SHorizontalBox::Slot()
 			.AutoWidth()
 			.Padding(2)
-			.VAlign(EVerticalAlignment::VAlign_Center)
 			[
-				SNew(SBox)
-				.HAlign(EHorizontalAlignment::HAlign_Right)
+				SNew(SButton)
+				.ButtonStyle(FEditorStyle::Get(), "HoverHintOnly")
+				.ForegroundColor(FSlateColor::UseForeground())
+				.HAlign(HAlign_Center)
+				.VAlign(VAlign_Center)
+				.OnClicked(this, &SRowEditor::OnMoveRowClicked, FDataTableEditorUtils::ERowMoveDirection::Up)
+				.ToolTipText(LOCTEXT("MoveUpTooltip", "Move the currently selected row up by one in the data table"))
 				[
-					SNew(STextBlock).Text(LOCTEXT("SelectRow", "Select:"))
+					SNew(STextBlock)
+					.Font(FEditorStyle::Get().GetFontStyle("FontAwesome.14"))
+					.Text(FText::FromString(FString(TEXT("\xf106"))) /*fa-angle-up*/)
 				]
 			]
 			+ SHorizontalBox::Slot()
 			.AutoWidth()
 			.Padding(2)
 			[
-				SNew(SBox)
-				.WidthOverride(2 * ButtonWidth)
+				SNew(SButton)
+				.ButtonStyle(FEditorStyle::Get(), "HoverHintOnly")
+				.ForegroundColor(FSlateColor::UseForeground())
+				.HAlign(HAlign_Center)
+				.VAlign(VAlign_Center)
+				.OnClicked(this, &SRowEditor::OnMoveRowClicked, FDataTableEditorUtils::ERowMoveDirection::Down)
+				.ToolTipText(LOCTEXT("MoveDownTooltip", "Move the currently selected row down by one in the data table"))
 				[
-					SAssignNew(RowComboBox, SComboBox<TSharedPtr<FName>>)
-					.OptionsSource(&CachedRowNames)
-					.OnSelectionChanged(this, &SRowEditor::OnSelectionChanged)
-					.OnGenerateWidget(this, &SRowEditor::OnGenerateWidget)
-					.Content()
-					[
-						SNew(STextBlock).Text(this, &SRowEditor::GetCurrentNameAsText)
-					]
+					SNew(STextBlock)
+					.Font(FEditorStyle::Get().GetFontStyle("FontAwesome.14"))
+					.Text(FText::FromString(FString(TEXT("\xf107"))) /*fa-angle-down*/)
+				]
+			]
+			+ SHorizontalBox::Slot()
+			.AutoWidth()
+			.Padding(2)
+			[
+				SNew(SButton)
+				.ButtonStyle(FEditorStyle::Get(), "HoverHintOnly")
+				.ForegroundColor(FSlateColor::UseForeground())
+				.HAlign(HAlign_Center)
+				.VAlign(VAlign_Center)
+				.OnClicked(this, &SRowEditor::OnMoveToExtentClicked, FDataTableEditorUtils::ERowMoveDirection::Up)
+				.ToolTipText(LOCTEXT("MoveToTopTooltip", "Move the currently selected row to the top of the data table"))
+				[
+					SNew(STextBlock)
+					.Font(FEditorStyle::Get().GetFontStyle("FontAwesome.14"))
+					.Text(FText::FromString(FString(TEXT("\xf102"))) /*fa-angle-double-up*/)
+				]
+			]
+			+ SHorizontalBox::Slot()
+			.AutoWidth()
+			.Padding(2)
+			[
+				SNew(SButton)
+				.ButtonStyle(FEditorStyle::Get(), "HoverHintOnly")
+				.ForegroundColor(FSlateColor::UseForeground())
+				.HAlign(HAlign_Center)
+				.VAlign(VAlign_Center)
+				.OnClicked(this, &SRowEditor::OnMoveToExtentClicked, FDataTableEditorUtils::ERowMoveDirection::Down)
+				.ToolTipText(LOCTEXT("MoveToBottomTooltip", "Move the currently selected row to the bottom of the data table"))
+				[
+					SNew(STextBlock)
+					.Font(FEditorStyle::Get().GetFontStyle("FontAwesome.14"))
+					.Text(FText::FromString(FString(TEXT("\xf103"))) /*fa-angle-double-down*/)
 				]
 			]
 		]
diff --git a/Engine/Source/Editor/DataTableEditor/Private/SRowEditor.h b/Engine/Source/Editor/DataTableEditor/Private/SRowEditor.h
--- a/Engine/Source/Editor/DataTableEditor/Private/SRowEditor.h
+++ b/Engine/Source/Editor/DataTableEditor/Private/SRowEditor.h
@@ -55,6 +55,8 @@ private:
 
 	FReply OnAddClicked();
 	FReply OnRemoveClicked();
+	FReply OnMoveRowClicked(FDataTableEditorUtils::ERowMoveDirection MoveDirection);
+	FReply OnMoveToExtentClicked(FDataTableEditorUtils::ERowMoveDirection MoveDirection);
 	void OnRowRenamed(const FText& Text, ETextCommit::Type CommitType);
 
 public:
diff --git a/Engine/Source/Editor/UnrealEd/Private/DataTableEditorUtils.cpp b/Engine/Source/Editor/UnrealEd/Private/DataTableEditorUtils.cpp
--- a/Engine/Source/Editor/UnrealEd/Private/DataTableEditorUtils.cpp
+++ b/Engine/Source/Editor/UnrealEd/Private/DataTableEditorUtils.cpp
@@ -78,6 +78,80 @@ bool FDataTableEditorUtils::RenameRow(UDataTable* DataTable, FName OldName, FNam
 	return bResult;
 }
 
+bool FDataTableEditorUtils::MoveRow(UDataTable* DataTable, FName RowName, ERowMoveDirection Direction, int32 NumRowsToMoveBy)
+{
+	if (!DataTable)
+	{
+		return false;
+	}
+	
+	// Our maps are ordered which is why we can get away with this
+	// If we ever change our map implementation, we'll need to preserve this order information in a separate array and 
+	// make sure that order dependent code (such as exporting and the data table viewer) use that when dealing with rows
+	// This may also require making RowMap private and fixing up all the existing code that references it directly
+	TArray<FName> OrderedRowNames;
+	DataTable->RowMap.GenerateKeyArray(OrderedRowNames);
+
+	const int32 CurrentRowIndex = OrderedRowNames.IndexOfByKey(RowName);
+	if (CurrentRowIndex == INDEX_NONE)
+	{
+		return false;
+	}
+	
+	// Calculate our new row index, clamped to the available rows
+	int32 NewRowIndex = INDEX_NONE;
+	switch(Direction)
+	{
+	case ERowMoveDirection::Up:
+		NewRowIndex = FMath::Clamp(CurrentRowIndex - NumRowsToMoveBy, 0, OrderedRowNames.Num() - 1);
+		break;
+
+	case ERowMoveDirection::Down:
+		NewRowIndex = FMath::Clamp(CurrentRowIndex + NumRowsToMoveBy, 0, OrderedRowNames.Num() - 1);
+		break;
+
+	default:
+		break;
+	}
+
+	if (NewRowIndex == INDEX_NONE)
+	{
+		return false;
+	}
+
+	if (CurrentRowIndex == NewRowIndex)
+	{
+		// Nothing to do, but not an error
+		return true;
+	}
+
+	// Swap the order around as requested
+	OrderedRowNames.RemoveAt(CurrentRowIndex, 1, false);
+	OrderedRowNames.Insert(RowName, NewRowIndex);
+
+	// Build a name -> index map as the KeySort will hit this a lot
+	TMap<FName, int32> NamesToNewIndex;
+	for (int32 NameIndex = 0; NameIndex < OrderedRowNames.Num(); ++NameIndex)
+	{
+		NamesToNewIndex.Add(OrderedRowNames[NameIndex], NameIndex);
+	}
+
+	BroadcastPreChange(DataTable, EDataTableChangeInfo::RowList);
+	DataTable->Modify();
+
+	// Re-sort the map keys to match the new order
+	DataTable->RowMap.KeySort([&NamesToNewIndex](const FName& One, const FName& Two) -> bool
+	{
+		const int32 OneIndex = NamesToNewIndex.FindRef(One);
+		const int32 TwoIndex = NamesToNewIndex.FindRef(Two);
+		return OneIndex < TwoIndex;
+	});
+
+	BroadcastPostChange(DataTable, EDataTableChangeInfo::RowList);
+
+	return true;
+}
+
 void FDataTableEditorUtils::BroadcastPreChange(UDataTable* DataTable, EDataTableChangeInfo Info)
 {
 	FDataTableEditorManager::Get().PreChange(DataTable, Info);
diff --git a/Engine/Source/Editor/UnrealEd/Public/DataTableEditorUtils.h b/Engine/Source/Editor/UnrealEd/Public/DataTableEditorUtils.h
--- a/Engine/Source/Editor/UnrealEd/Public/DataTableEditorUtils.h
+++ b/Engine/Source/Editor/UnrealEd/Public/DataTableEditorUtils.h
@@ -12,6 +12,12 @@ struct UNREALED_API FDataTableEditorUtils
 		RowList,
 	};
 
+	enum class ERowMoveDirection
+	{
+		Up,
+		Down,
+	};
+
 	class FDataTableEditorManager : public FListenerManager < UDataTable, EDataTableChangeInfo >
 	{
 		FDataTableEditorManager() {}
@@ -28,6 +34,7 @@ struct UNREALED_API FDataTableEditorUtils
 	static bool RemoveRow(UDataTable* DataTable, FName Name);
 	static uint8* AddRow(UDataTable* DataTable, FName RowName);
 	static bool RenameRow(UDataTable* DataTable, FName OldName, FName NewName);
+	static bool MoveRow(UDataTable* DataTable, FName RowName, ERowMoveDirection Direction, int32 NumRowsToMoveBy = 1);
 
 	static void BroadcastPreChange(UDataTable* DataTable, EDataTableChangeInfo Info);
 	static void BroadcastPostChange(UDataTable* DataTable, EDataTableChangeInfo Info);
