Add includes to files to remove the need of including Engine.h.

diff --git a/Engine/Plugins/2D/Paper2D/Source/Paper2D/Private/PaperFlipbookComponent.cpp b/Engine/Plugins/2D/Paper2D/Source/Paper2D/Private/PaperFlipbookComponent.cpp
--- a/Engine/Plugins/2D/Paper2D/Source/Paper2D/Private/PaperFlipbookComponent.cpp
+++ b/Engine/Plugins/2D/Paper2D/Source/Paper2D/Private/PaperFlipbookComponent.cpp
@@ -9,6 +9,7 @@
 #include "Runtime/Core/Public/Logging/MessageLog.h"
 #include "Runtime/Core/Public/Misc/MapErrors.h"
 #include "Runtime/CoreUObject/Public/Misc/UObjectToken.h"
+#include "PhysicsEngine/BodySetup.h"
 
 #define LOCTEXT_NAMESPACE "Paper2D"
 
diff --git a/Engine/Plugins/2D/Paper2D/Source/Paper2D/Private/PaperRenderSceneProxy.cpp b/Engine/Plugins/2D/Paper2D/Source/Paper2D/Private/PaperRenderSceneProxy.cpp
--- a/Engine/Plugins/2D/Paper2D/Source/Paper2D/Private/PaperRenderSceneProxy.cpp
+++ b/Engine/Plugins/2D/Paper2D/Source/Paper2D/Private/PaperRenderSceneProxy.cpp
@@ -7,6 +7,10 @@
 #include "Rendering/PaperBatchSceneProxy.h"
 #include "PhysicsEngine/BodySetup2D.h"
 #include "LocalVertexFactory.h"
+#include "MeshBatch.h"
+#include "EngineGlobals.h"
+#include "SceneManagement.h"
+#include "Engine/Engine.h"
 
 //////////////////////////////////////////////////////////////////////////
 // FPaperSpriteVertex
diff --git a/Engine/Plugins/2D/Paper2D/Source/Paper2D/Private/PaperSpriteSceneProxy.cpp b/Engine/Plugins/2D/Paper2D/Source/Paper2D/Private/PaperSpriteSceneProxy.cpp
--- a/Engine/Plugins/2D/Paper2D/Source/Paper2D/Private/PaperSpriteSceneProxy.cpp
+++ b/Engine/Plugins/2D/Paper2D/Source/Paper2D/Private/PaperSpriteSceneProxy.cpp
@@ -3,6 +3,9 @@
 #include "Paper2DPrivatePCH.h"
 #include "PaperSpriteSceneProxy.h"
 #include "PhysicsEngine/BodySetup2D.h"
+#include "SceneManagement.h"
+#include "EngineGlobals.h"
+#include "Engine/Engine.h"
 
 //////////////////////////////////////////////////////////////////////////
 // FPaperSpriteSceneProxy
diff --git a/Engine/Plugins/2D/Paper2D/Source/Paper2D/Private/PaperTileLayer.cpp b/Engine/Plugins/2D/Paper2D/Source/Paper2D/Private/PaperTileLayer.cpp
--- a/Engine/Plugins/2D/Paper2D/Source/Paper2D/Private/PaperTileLayer.cpp
+++ b/Engine/Plugins/2D/Paper2D/Source/Paper2D/Private/PaperTileLayer.cpp
@@ -3,6 +3,8 @@
 #include "Paper2DPrivatePCH.h"
 #include "SpriteEditorOnlyTypes.h"
 #include "ComponentReregisterContext.h"
+#include "PhysicsEngine/BoxElem.h"
+#include "PhysicsEngine/BodySetup.h"
 
 //////////////////////////////////////////////////////////////////////////
 // FTileMapLayerReregisterContext
diff --git a/Engine/Plugins/2D/Paper2D/Source/Paper2D/Private/PaperTileMapComponent.cpp b/Engine/Plugins/2D/Paper2D/Source/Paper2D/Private/PaperTileMapComponent.cpp
--- a/Engine/Plugins/2D/Paper2D/Source/Paper2D/Private/PaperTileMapComponent.cpp
+++ b/Engine/Plugins/2D/Paper2D/Source/Paper2D/Private/PaperTileMapComponent.cpp
@@ -4,6 +4,7 @@
 #include "PaperTileMapRenderSceneProxy.h"
 #include "PhysicsEngine/BodySetup.h"
 #include "PaperCustomVersion.h"
+#include "Engine/Canvas.h"
 
 #define LOCTEXT_NAMESPACE "Paper2D"
 
diff --git a/Engine/Plugins/2D/Paper2D/Source/Paper2D/Private/PaperTileMapRenderSceneProxy.cpp b/Engine/Plugins/2D/Paper2D/Source/Paper2D/Private/PaperTileMapRenderSceneProxy.cpp
--- a/Engine/Plugins/2D/Paper2D/Source/Paper2D/Private/PaperTileMapRenderSceneProxy.cpp
+++ b/Engine/Plugins/2D/Paper2D/Source/Paper2D/Private/PaperTileMapRenderSceneProxy.cpp
@@ -4,6 +4,9 @@
 #include "PaperTileMapRenderSceneProxy.h"
 #include "PhysicsEngine/BodySetup.h"
 #include "PhysicsEngine/BodySetup2D.h"
+#include "EngineGlobals.h"
+#include "SceneManagement.h"
+#include "Engine/Engine.h"
 
 //////////////////////////////////////////////////////////////////////////
 // FPaperTileMapRenderSceneProxy
diff --git a/Engine/Plugins/2D/Paper2D/Source/Paper2D/Private/Terrain/PaperTerrainComponent.cpp b/Engine/Plugins/2D/Paper2D/Source/Paper2D/Private/Terrain/PaperTerrainComponent.cpp
--- a/Engine/Plugins/2D/Paper2D/Source/Paper2D/Private/Terrain/PaperTerrainComponent.cpp
+++ b/Engine/Plugins/2D/Paper2D/Source/Paper2D/Private/Terrain/PaperTerrainComponent.cpp
@@ -5,6 +5,8 @@
 
 #include "PaperRenderSceneProxy.h"
 #include "PaperGeomTools.h"
+#include "PhysicsEngine/BodySetup.h"
+#include "PhysicsEngine/ConvexElem.h"
 
 #define PAPER_USE_MATERIAL_SLOPES 1
 
diff --git a/Engine/Plugins/2D/Paper2D/Source/Paper2DEditor/Private/ExtractSprites/SPaperExtractSpritesDialog.cpp b/Engine/Plugins/2D/Paper2D/Source/Paper2DEditor/Private/ExtractSprites/SPaperExtractSpritesDialog.cpp
--- a/Engine/Plugins/2D/Paper2D/Source/Paper2DEditor/Private/ExtractSprites/SPaperExtractSpritesDialog.cpp
+++ b/Engine/Plugins/2D/Paper2D/Source/Paper2DEditor/Private/ExtractSprites/SPaperExtractSpritesDialog.cpp
@@ -9,6 +9,7 @@
 #include "Editor/PropertyEditor/Public/IDetailsView.h"
 #include "ContentBrowserModule.h"
 #include "AssetToolsModule.h"
+#include "CanvasTypes.h"
 
 #define LOCTEXT_NAMESPACE "PaperEditor"
 
diff --git a/Engine/Plugins/2D/Paper2D/Source/Paper2DEditor/Private/PaperEditorShared/SocketEditing.cpp b/Engine/Plugins/2D/Paper2D/Source/Paper2DEditor/Private/PaperEditorShared/SocketEditing.cpp
--- a/Engine/Plugins/2D/Paper2D/Source/Paper2DEditor/Private/PaperEditorShared/SocketEditing.cpp
+++ b/Engine/Plugins/2D/Paper2D/Source/Paper2DEditor/Private/PaperEditorShared/SocketEditing.cpp
@@ -3,6 +3,8 @@
 #include "Paper2DEditorPrivatePCH.h"
 #include "SocketEditing.h"
 #include "SpriteGeometryEditMode.h"
+#include "CanvasItem.h"
+#include "CanvasTypes.h"
 
 //////////////////////////////////////////////////////////////////////////
 // HSpriteSelectableObjectHitProxy
diff --git a/Engine/Plugins/2D/Paper2D/Source/Paper2DEditor/Private/TileMapEditing/TileMapEditorViewportClient.cpp b/Engine/Plugins/2D/Paper2D/Source/Paper2DEditor/Private/TileMapEditing/TileMapEditorViewportClient.cpp
--- a/Engine/Plugins/2D/Paper2D/Source/Paper2DEditor/Private/TileMapEditing/TileMapEditorViewportClient.cpp
+++ b/Engine/Plugins/2D/Paper2D/Source/Paper2DEditor/Private/TileMapEditing/TileMapEditorViewportClient.cpp
@@ -9,6 +9,7 @@
 #include "ScopedTransaction.h"
 #include "Runtime/Engine/Public/ComponentReregisterContext.h"
 #include "CanvasTypes.h"
+#include "Engine/Selection.h"
 
 #define LOCTEXT_NAMESPACE "TileMapEditor"
 
diff --git a/Engine/Plugins/Experimental/LiveEditor/Source/LiveEditor/Private/LiveEditorConfigWindow.cpp b/Engine/Plugins/Experimental/LiveEditor/Source/LiveEditor/Private/LiveEditorConfigWindow.cpp
--- a/Engine/Plugins/Experimental/LiveEditor/Source/LiveEditor/Private/LiveEditorConfigWindow.cpp
+++ b/Engine/Plugins/Experimental/LiveEditor/Source/LiveEditor/Private/LiveEditorConfigWindow.cpp
@@ -7,6 +7,7 @@
 #include "LiveEditorDeviceSetupWindow.h"
 #include "LiveEditorWizardBase.h"
 #include "SExpandableArea.h"
+#include "Engine/Selection.h"
 
 DEFINE_LOG_CATEGORY_STATIC(LiveEditorConfigWindow, Log, All);
 
diff --git a/Engine/Plugins/Experimental/LiveEditor/Source/LiveEditor/Private/LiveEditorKismetLibrary.cpp b/Engine/Plugins/Experimental/LiveEditor/Source/LiveEditor/Private/LiveEditorKismetLibrary.cpp
--- a/Engine/Plugins/Experimental/LiveEditor/Source/LiveEditor/Private/LiveEditorKismetLibrary.cpp
+++ b/Engine/Plugins/Experimental/LiveEditor/Source/LiveEditor/Private/LiveEditorKismetLibrary.cpp
@@ -3,6 +3,7 @@
 #include "LiveEditorPrivatePCH.h"
 #include "LiveEditorKismetLibrary.h"
 #include "DefaultValueHelper.h"
+#include "Engine/Selection.h"
 
 namespace nLiveEditorKismetLibrary
 {
diff --git a/Engine/Plugins/Runtime/CableComponent/Source/CableComponent/Private/CableComponent.cpp b/Engine/Plugins/Runtime/CableComponent/Source/CableComponent/Private/CableComponent.cpp
--- a/Engine/Plugins/Runtime/CableComponent/Source/CableComponent/Private/CableComponent.cpp
+++ b/Engine/Plugins/Runtime/CableComponent/Source/CableComponent/Private/CableComponent.cpp
@@ -2,6 +2,9 @@
 
 #include "CableComponentPluginPrivatePCH.h"
 #include "DynamicMeshBuilder.h"
+#include "EngineGlobals.h"
+#include "LocalVertexFactory.h"
+#include "Engine/Engine.h"
 
 
 /** Vertex Buffer */
diff --git a/Engine/Plugins/Runtime/CustomMeshComponent/Source/CustomMeshComponent/Private/CustomMeshComponent.cpp b/Engine/Plugins/Runtime/CustomMeshComponent/Source/CustomMeshComponent/Private/CustomMeshComponent.cpp
--- a/Engine/Plugins/Runtime/CustomMeshComponent/Source/CustomMeshComponent/Private/CustomMeshComponent.cpp
+++ b/Engine/Plugins/Runtime/CustomMeshComponent/Source/CustomMeshComponent/Private/CustomMeshComponent.cpp
@@ -3,6 +3,10 @@
 #include "CustomMeshComponentPluginPrivatePCH.h"
 #include "CustomMeshComponent.h"
 #include "DynamicMeshBuilder.h"
+#include "MaterialShared.h"
+#include "EngineGlobals.h"
+#include "LocalVertexFactory.h"
+#include "Engine/Engine.h"
 
 /** Vertex Buffer */
 class FCustomMeshVertexBuffer : public FVertexBuffer 
diff --git a/Engine/Source/Developer/HotReload/Private/HotReloadClassReinstancer.cpp b/Engine/Source/Developer/HotReload/Private/HotReloadClassReinstancer.cpp
--- a/Engine/Source/Developer/HotReload/Private/HotReloadClassReinstancer.cpp
+++ b/Engine/Source/Developer/HotReload/Private/HotReloadClassReinstancer.cpp
@@ -4,6 +4,7 @@
 #include "HotReloadClassReinstancer.h"
 
 #if WITH_ENGINE
+#include "Engine/BlueprintGeneratedClass.h"
 
 void FHotReloadClassReinstancer::SetupNewClassReinstancing(UClass* InNewClass, UClass* InOldClass)
 {
@@ -492,4 +493,4 @@ void FHotReloadClassReinstancer::AddReferencedObjects(FReferenceCollector& Colle
 	Collector.AddReferencedObject(CopyOfPreviousCDO);
 }
 
-#endif
\ No newline at end of file
+#endif
diff --git a/Engine/Source/Developer/Merge/Private/SMergeAssetPickerView.cpp b/Engine/Source/Developer/Merge/Private/SMergeAssetPickerView.cpp
--- a/Engine/Source/Developer/Merge/Private/SMergeAssetPickerView.cpp
+++ b/Engine/Source/Developer/Merge/Private/SMergeAssetPickerView.cpp
@@ -5,6 +5,7 @@
 #include "Editor/ContentBrowser/Public/ContentBrowserModule.h"
 #include "SBlueprintRevisionMenu.h"
 #include "Editor.h"
+#include "Engine/Selection.h"
 
 #define LOCTEXT_NAMESPACE "SMergeAssetPickerView"
 
diff --git a/Engine/Source/Editor/AnimGraph/Classes/AnimGraphNode_SkeletalControlBase.h b/Engine/Source/Editor/AnimGraph/Classes/AnimGraphNode_SkeletalControlBase.h
--- a/Engine/Source/Editor/AnimGraph/Classes/AnimGraphNode_SkeletalControlBase.h
+++ b/Engine/Source/Editor/AnimGraph/Classes/AnimGraphNode_SkeletalControlBase.h
@@ -5,6 +5,8 @@
 #include "AnimGraphNode_Base.h"
 #include "AnimGraphNode_SkeletalControlBase.generated.h"
 
+struct HActor;
+
 UCLASS(Abstract, MinimalAPI)
 class UAnimGraphNode_SkeletalControlBase : public UAnimGraphNode_Base
 {
diff --git a/Engine/Source/Editor/BehaviorTreeEditor/Private/BehaviorTreeDebugger.cpp b/Engine/Source/Editor/BehaviorTreeEditor/Private/BehaviorTreeDebugger.cpp
--- a/Engine/Source/Editor/BehaviorTreeEditor/Private/BehaviorTreeDebugger.cpp
+++ b/Engine/Source/Editor/BehaviorTreeEditor/Private/BehaviorTreeDebugger.cpp
@@ -8,6 +8,8 @@
 #include "BehaviorTreeDelegates.h"
 #include "GameplayDebuggingReplicator.h"
 #include "GameplayDebuggingComponent.h"
+#include "Engine/Selection.h"
+#include "EngineUtils.h"
 
 FBehaviorTreeDebugger::FBehaviorTreeDebugger()
 {
diff --git a/Engine/Source/Editor/BehaviorTreeEditor/Private/BehaviorTreeEditor.cpp b/Engine/Source/Editor/BehaviorTreeEditor/Private/BehaviorTreeEditor.cpp
--- a/Engine/Source/Editor/BehaviorTreeEditor/Private/BehaviorTreeEditor.cpp
+++ b/Engine/Source/Editor/BehaviorTreeEditor/Private/BehaviorTreeEditor.cpp
@@ -35,6 +35,7 @@
 #include "ContentBrowserModule.h"
 #include "SDockTab.h"
 #include "GenericCommands.h"
+#include "Engine/BlueprintGeneratedClass.h"
 
 #define LOCTEXT_NAMESPACE "BehaviorTreeEditor"
 
diff --git a/Engine/Source/Editor/Blutility/Private/PlacedEditorUtilityBase.cpp b/Engine/Source/Editor/Blutility/Private/PlacedEditorUtilityBase.cpp
--- a/Engine/Source/Editor/Blutility/Private/PlacedEditorUtilityBase.cpp
+++ b/Engine/Source/Editor/Blutility/Private/PlacedEditorUtilityBase.cpp
@@ -2,6 +2,7 @@
 
 #include "BlutilityPrivatePCH.h"
 #include "PlacedEditorUtilityBase.h"
+#include "Engine/Selection.h"
 
 /////////////////////////////////////////////////////
 // APlacedEditorUtilityBase
diff --git a/Engine/Source/Editor/ClassViewer/Private/SClassViewer.cpp b/Engine/Source/Editor/ClassViewer/Private/SClassViewer.cpp
--- a/Engine/Source/Editor/ClassViewer/Private/SClassViewer.cpp
+++ b/Engine/Source/Editor/ClassViewer/Private/SClassViewer.cpp
@@ -35,6 +35,7 @@
 #include "SSearchBox.h"
 
 #include "SListViewSelectorDropdownMenu.h"
+#include "Engine/BlueprintGeneratedClass.h"
 
 #define LOCTEXT_NAMESPACE "SClassViewer"
 
diff --git a/Engine/Source/Editor/DataTableEditor/Private/DataTableEditor.cpp b/Engine/Source/Editor/DataTableEditor/Private/DataTableEditor.cpp
--- a/Engine/Source/Editor/DataTableEditor/Private/DataTableEditor.cpp
+++ b/Engine/Source/Editor/DataTableEditor/Private/DataTableEditor.cpp
@@ -10,6 +10,7 @@
 #include "SRowEditor.h"
 #include "Engine/DataTable.h"
 #include "Json.h"
+#include "Engine/UserDefinedStruct.h"
  
 #define LOCTEXT_NAMESPACE "DataTableEditor"
 
@@ -44,7 +45,7 @@ FDataTableEditor::~FDataTableEditor()
 	if (DataTable.IsValid())
 	{
 		SaveLayoutData();
-	}
+}
 }
 
 void FDataTableEditor::PreChange(const class UUserDefinedStruct* Struct, FStructureEditorUtils::EStructureEditorChangeInfo Info)
@@ -267,7 +268,7 @@ TSharedRef<SWidget> FDataTableEditor::CreateGridPanel()
 							ColumnWidth.CurrentWidth = static_cast<float>(LayoutColumnWidth);
 						}
 					}
-				}
+		}
 			}
 		}
 
@@ -445,12 +446,12 @@ TSharedRef<SVerticalBox> FDataTableEditor::CreateContentBox()
 				.ExternalScrollbar(HorizontalScrollBar)
 				+SScrollBox::Slot()
 				[
-					SAssignNew(ScrollBoxWidget, SScrollBox)
+			SAssignNew(ScrollBoxWidget, SScrollBox)
 					.Orientation(Orient_Vertical)
 					.ExternalScrollbar(VerticalScrollBar)
 					.ConsumeMouseWheel(EConsumeMouseWheel::Always) // Always consume the mouse wheel events to prevent the outer scroll box from scrolling
-					+SScrollBox::Slot()
-					[
+			+SScrollBox::Slot()
+			[
 						CreateGridPanel()
 					]
 				]
diff --git a/Engine/Source/Editor/DataTableEditor/Private/SRowEditor.cpp b/Engine/Source/Editor/DataTableEditor/Private/SRowEditor.cpp
--- a/Engine/Source/Editor/DataTableEditor/Private/SRowEditor.cpp
+++ b/Engine/Source/Editor/DataTableEditor/Private/SRowEditor.cpp
@@ -8,6 +8,7 @@
 #include "IDetailsView.h"
 #include "DataTableUtils.h"
 #include "DataTableEditorUtils.h"
+#include "Engine/UserDefinedStruct.h"
 
 #define LOCTEXT_NAMESPACE "SRowEditor"
 
diff --git a/Engine/Source/Editor/DetailCustomizations/Private/ActorDetails.cpp b/Engine/Source/Editor/DetailCustomizations/Private/ActorDetails.cpp
--- a/Engine/Source/Editor/DetailCustomizations/Private/ActorDetails.cpp
+++ b/Engine/Source/Editor/DetailCustomizations/Private/ActorDetails.cpp
@@ -30,6 +30,11 @@
 #include "SHyperlink.h"
 #include "SNotificationList.h"
 #include "NotificationManager.h"
+#include "GameFramework/Volume.h"
+#include "Engine/Selection.h"
+#include "GameFramework/WorldSettings.h"
+#include "Engine/BlueprintGeneratedClass.h"
+#include "Components/BillboardComponent.h"
 
 #define LOCTEXT_NAMESPACE "ActorDetails"
 
diff --git a/Engine/Source/Editor/DetailCustomizations/Private/HierarchicalSimplificationCustomizations.cpp b/Engine/Source/Editor/DetailCustomizations/Private/HierarchicalSimplificationCustomizations.cpp
--- a/Engine/Source/Editor/DetailCustomizations/Private/HierarchicalSimplificationCustomizations.cpp
+++ b/Engine/Source/Editor/DetailCustomizations/Private/HierarchicalSimplificationCustomizations.cpp
@@ -3,6 +3,7 @@
 #include "DetailCustomizationsPrivatePCH.h"
 #include "HierarchicalSimplificationCustomizations.h"
 #include "Developer/MeshUtilities/Public/MeshUtilities.h"
+#include "GameFramework/WorldSettings.h"
 
 TSharedRef<IPropertyTypeCustomization> FHierarchicalSimplificationCustomizations::MakeInstance() 
 {
diff --git a/Engine/Source/Editor/DetailCustomizations/Private/SkyLightComponentDetails.cpp b/Engine/Source/Editor/DetailCustomizations/Private/SkyLightComponentDetails.cpp
--- a/Engine/Source/Editor/DetailCustomizations/Private/SkyLightComponentDetails.cpp
+++ b/Engine/Source/Editor/DetailCustomizations/Private/SkyLightComponentDetails.cpp
@@ -4,6 +4,7 @@
 #include "SkyLightComponentDetails.h"
 #include "Components/LightComponentBase.h"
 #include "Engine/SkyLight.h"
+#include "Components/SkyLightComponent.h"
 
 #define LOCTEXT_NAMESPACE "SkyLightComponentDetails"
 
diff --git a/Engine/Source/Editor/DetailCustomizations/Private/TextureLODSettingsDetails.cpp b/Engine/Source/Editor/DetailCustomizations/Private/TextureLODSettingsDetails.cpp
--- a/Engine/Source/Editor/DetailCustomizations/Private/TextureLODSettingsDetails.cpp
+++ b/Engine/Source/Editor/DetailCustomizations/Private/TextureLODSettingsDetails.cpp
@@ -2,6 +2,8 @@
 
 #include "DetailCustomizationsPrivatePCH.h"
 #include "TextureLODSettingsDetails.h"
+#include "Engine/TextureLODSettings.h"
+#include "DeviceProfiles/DeviceProfile.h"
 
 #define LOCTEXT_NAMESPACE "TextureLODSettingsDetails"
 
diff --git a/Engine/Source/Editor/DetailCustomizations/Private/TextureLODSettingsDetails.h b/Engine/Source/Editor/DetailCustomizations/Private/TextureLODSettingsDetails.h
--- a/Engine/Source/Editor/DetailCustomizations/Private/TextureLODSettingsDetails.h
+++ b/Engine/Source/Editor/DetailCustomizations/Private/TextureLODSettingsDetails.h
@@ -2,6 +2,7 @@
 
 #pragma once
 
+struct FTextureLODGroup;
 
 /**
  * Texture Group layout for property editor views.
diff --git a/Engine/Source/Editor/FoliageEdit/Private/ActorFactoryProceduralFoliage.cpp b/Engine/Source/Editor/FoliageEdit/Private/ActorFactoryProceduralFoliage.cpp
--- a/Engine/Source/Editor/FoliageEdit/Private/ActorFactoryProceduralFoliage.cpp
+++ b/Engine/Source/Editor/FoliageEdit/Private/ActorFactoryProceduralFoliage.cpp
@@ -9,6 +9,7 @@ ActorFactory.cpp:
 #include "ProceduralFoliageSpawner.h"
 #include "ProceduralFoliageVolume.h"
 #include "ProceduralFoliageComponent.h"
+#include "AssetData.h"
 
 DEFINE_LOG_CATEGORY_STATIC(LogActorFactory, Log, All);
 
diff --git a/Engine/Source/Editor/FoliageEdit/Private/FoliageEdMode.cpp b/Engine/Source/Editor/FoliageEdit/Private/FoliageEdMode.cpp
--- a/Engine/Source/Editor/FoliageEdit/Private/FoliageEdMode.cpp
+++ b/Engine/Source/Editor/FoliageEdit/Private/FoliageEdMode.cpp
@@ -33,6 +33,7 @@
 #include "EngineUtils.h"
 #include "LandscapeDataAccess.h"
 #include "FoliageType_InstancedStaticMesh.h"
+#include "Components/BrushComponent.h"
 
 
 #define LOCTEXT_NAMESPACE "FoliageEdMode"
diff --git a/Engine/Source/Editor/GameProjectGeneration/Private/SNewClassDialog.cpp b/Engine/Source/Editor/GameProjectGeneration/Private/SNewClassDialog.cpp
--- a/Engine/Source/Editor/GameProjectGeneration/Private/SNewClassDialog.cpp
+++ b/Engine/Source/Editor/GameProjectGeneration/Private/SNewClassDialog.cpp
@@ -30,6 +30,7 @@
 #include "ContentBrowserModule.h"
 #include "SNotificationList.h"
 #include "NotificationManager.h"
+#include "Engine/BlueprintGeneratedClass.h"
 
 
 #define LOCTEXT_NAMESPACE "GameProjectGeneration"
diff --git a/Engine/Source/Editor/GraphEditor/Private/AnimationStateNodes/SGraphNodeAnimTransition.h b/Engine/Source/Editor/GraphEditor/Private/AnimationStateNodes/SGraphNodeAnimTransition.h
--- a/Engine/Source/Editor/GraphEditor/Private/AnimationStateNodes/SGraphNodeAnimTransition.h
+++ b/Engine/Source/Editor/GraphEditor/Private/AnimationStateNodes/SGraphNodeAnimTransition.h
@@ -1,6 +1,9 @@
 // Copyright 1998-2015 Epic Games, Inc. All Rights Reserved.
 
 #pragma once
+#include "SGraphNode.h"
+
+class UAnimStateTransitionNode;
 
 class SGraphNodeAnimTransition : public SGraphNode
 {
diff --git a/Engine/Source/Editor/GraphEditor/Private/KismetNodes/KismetNodeInfoContext.cpp b/Engine/Source/Editor/GraphEditor/Private/KismetNodes/KismetNodeInfoContext.cpp
--- a/Engine/Source/Editor/GraphEditor/Private/KismetNodes/KismetNodeInfoContext.cpp
+++ b/Engine/Source/Editor/GraphEditor/Private/KismetNodes/KismetNodeInfoContext.cpp
@@ -3,6 +3,7 @@
 #include "GraphEditorCommon.h"
 #include "KismetNodes/KismetNodeInfoContext.h"
 #include "Editor/UnrealEd/Public/Kismet2/BlueprintEditorUtils.h"
+#include "Engine/BlueprintGeneratedClass.h"
 
 //////////////////////////////////////////////////////////////////////////
 // FKismetNodeInfoContext
diff --git a/Engine/Source/Editor/GraphEditor/Private/KismetNodes/SGraphNodeCallParameterCollectionFunction.cpp b/Engine/Source/Editor/GraphEditor/Private/KismetNodes/SGraphNodeCallParameterCollectionFunction.cpp
--- a/Engine/Source/Editor/GraphEditor/Private/KismetNodes/SGraphNodeCallParameterCollectionFunction.cpp
+++ b/Engine/Source/Editor/GraphEditor/Private/KismetNodes/SGraphNodeCallParameterCollectionFunction.cpp
@@ -8,6 +8,7 @@
 #include "SGraphNodeK2Default.h"
 #include "SGraphNodeCallParameterCollectionFunction.h"
 #include "SGraphPinNameList.h"
+#include "Materials/MaterialParameterCollection.h"
 
 //////////////////////////////////////////////////////////////////////////
 // SGraphNodeCallParameterCollectionFunction
diff --git a/Engine/Source/Editor/GraphEditor/Private/KismetNodes/SGraphNodeK2Timeline.cpp b/Engine/Source/Editor/GraphEditor/Private/KismetNodes/SGraphNodeK2Timeline.cpp
--- a/Engine/Source/Editor/GraphEditor/Private/KismetNodes/SGraphNodeK2Timeline.cpp
+++ b/Engine/Source/Editor/GraphEditor/Private/KismetNodes/SGraphNodeK2Timeline.cpp
@@ -7,6 +7,7 @@
 #include "Editor/UnrealEd/Public/Kismet2/KismetDebugUtilities.h"
 #include "Editor/UnrealEd/Public/Kismet2/BlueprintEditorUtils.h"
 #include "KismetNodeInfoContext.h"
+#include "Components/TimelineComponent.h"
 
 #define LOCTEXT_NAMESPACE "SGraphNodeK2Base"
 
diff --git a/Engine/Source/Editor/GraphEditor/Private/KismetPins/SGraphPinObject.cpp b/Engine/Source/Editor/GraphEditor/Private/KismetPins/SGraphPinObject.cpp
--- a/Engine/Source/Editor/GraphEditor/Private/KismetPins/SGraphPinObject.cpp
+++ b/Engine/Source/Editor/GraphEditor/Private/KismetPins/SGraphPinObject.cpp
@@ -6,6 +6,7 @@
 #include "Editor/ContentBrowser/Public/ContentBrowserModule.h"
 #include "AssetData.h"
 #include "Editor/UnrealEd/Public/ScopedTransaction.h"
+#include "Engine/Selection.h"
 
 #define LOCTEXT_NAMESPACE "SGraphPinObject"
 
diff --git a/Engine/Source/Editor/GraphEditor/Private/MaterialNodes/SGraphNodeMaterialBase.cpp b/Engine/Source/Editor/GraphEditor/Private/MaterialNodes/SGraphNodeMaterialBase.cpp
--- a/Engine/Source/Editor/GraphEditor/Private/MaterialNodes/SGraphNodeMaterialBase.cpp
+++ b/Engine/Source/Editor/GraphEditor/Private/MaterialNodes/SGraphNodeMaterialBase.cpp
@@ -7,6 +7,8 @@
 #include "ScopedTransaction.h"
 #include "Runtime/Engine/Public/Slate/SceneViewport.h"
 #include "TutorialMetaData.h"
+#include "CanvasTypes.h"
+#include "CanvasItem.h"
 
 
 
diff --git a/Engine/Source/Editor/HardwareTargeting/Private/HardwareTargetingModule.cpp b/Engine/Source/Editor/HardwareTargeting/Private/HardwareTargetingModule.cpp
--- a/Engine/Source/Editor/HardwareTargeting/Private/HardwareTargetingModule.cpp
+++ b/Engine/Source/Editor/HardwareTargeting/Private/HardwareTargetingModule.cpp
@@ -10,6 +10,8 @@
 
 #include "Runtime/Engine/Classes/Engine/RendererSettings.h"
 #include "Editor/Documentation/Public/IDocumentation.h"
+#include "GameFramework/InputSettings.h"
+#include "GameMapsSettings.h"
 
 
 #define LOCTEXT_NAMESPACE "HardwareTargeting"
diff --git a/Engine/Source/Editor/IntroTutorials/Private/ClassTypeActions_EditorTutorial.cpp b/Engine/Source/Editor/IntroTutorials/Private/ClassTypeActions_EditorTutorial.cpp
--- a/Engine/Source/Editor/IntroTutorials/Private/ClassTypeActions_EditorTutorial.cpp
+++ b/Engine/Source/Editor/IntroTutorials/Private/ClassTypeActions_EditorTutorial.cpp
@@ -3,6 +3,7 @@
 #include "IntroTutorialsPrivatePCH.h"
 #include "ClassTypeActions_EditorTutorial.h"
 #include "EditorTutorial.h"
+#include "AssetData.h"
 
 #define LOCTEXT_NAMESPACE "IntroTutorials"
 
diff --git a/Engine/Source/Editor/Kismet/Private/BlueprintDetailsCustomization.cpp b/Engine/Source/Editor/Kismet/Private/BlueprintDetailsCustomization.cpp
--- a/Engine/Source/Editor/Kismet/Private/BlueprintDetailsCustomization.cpp
+++ b/Engine/Source/Editor/Kismet/Private/BlueprintDetailsCustomization.cpp
@@ -41,6 +41,7 @@
 #include "Engine/SCS_Node.h"
 #include "Components/TimelineComponent.h"
 #include "Engine/BlueprintGeneratedClass.h"
+#include "Components/ChildActorComponent.h"
 
 #define LOCTEXT_NAMESPACE "BlueprintDetailsCustomization"
 
diff --git a/Engine/Source/Editor/Kismet/Private/BlueprintEditor.cpp b/Engine/Source/Editor/Kismet/Private/BlueprintEditor.cpp
--- a/Engine/Source/Editor/Kismet/Private/BlueprintEditor.cpp
+++ b/Engine/Source/Editor/Kismet/Private/BlueprintEditor.cpp
@@ -91,6 +91,9 @@
 #include "SNotificationList.h"
 #include "NotificationManager.h"
 #include "NativeCodeGenerationTool.h"
+#include "GeneralProjectSettings.h"
+#include "Engine/LevelStreaming.h"
+#include "Engine/LevelStreamingKismet.h"
 
 #define LOCTEXT_NAMESPACE "BlueprintEditor"
 
diff --git a/Engine/Source/Editor/Kismet/Private/SBlueprintPalette.cpp b/Engine/Source/Editor/Kismet/Private/SBlueprintPalette.cpp
--- a/Engine/Source/Editor/Kismet/Private/SBlueprintPalette.cpp
+++ b/Engine/Source/Editor/Kismet/Private/SBlueprintPalette.cpp
@@ -35,6 +35,8 @@
 #include "SInlineEditableTextBlock.h"
 #include "Engine/SimpleConstructionScript.h"
 #include "Editor/UnrealEd/Public/Kismet2/ComponentEditorUtils.h"
+#include "Engine/SCS_Node.h"
+#include "Components/TimelineComponent.h"
 
 #define LOCTEXT_NAMESPACE "BlueprintPalette"
 
diff --git a/Engine/Source/Editor/Kismet/Private/SSCSEditor.cpp b/Engine/Source/Editor/Kismet/Private/SSCSEditor.cpp
--- a/Engine/Source/Editor/Kismet/Private/SSCSEditor.cpp
+++ b/Engine/Source/Editor/Kismet/Private/SSCSEditor.cpp
@@ -51,6 +51,7 @@
 
 #include "SourceCodeNavigation.h"
 #include "Editor/BlueprintGraph/Public/BlueprintEditorSettings.h"
+#include "Engine/BlueprintGeneratedClass.h"
 
 #define LOCTEXT_NAMESPACE "SSCSEditor"
 
diff --git a/Engine/Source/Editor/KismetCompiler/Private/AnimBlueprintCompiler.h b/Engine/Source/Editor/KismetCompiler/Private/AnimBlueprintCompiler.h
--- a/Engine/Source/Editor/KismetCompiler/Private/AnimBlueprintCompiler.h
+++ b/Engine/Source/Editor/KismetCompiler/Private/AnimBlueprintCompiler.h
@@ -12,6 +12,7 @@
 class UAnimationGraphSchema;
 class UAnimGraphNode_SaveCachedPose;
 class UAnimGraphNode_UseCachedPose;
+class UBlueprintGeneratedClass;
 struct FPoseLinkMappingRecord;
 
 //////////////////////////////////////////////////////////////////////////
diff --git a/Engine/Source/Editor/KismetCompiler/Private/KismetCompiler.cpp b/Engine/Source/Editor/KismetCompiler/Private/KismetCompiler.cpp
--- a/Engine/Source/Editor/KismetCompiler/Private/KismetCompiler.cpp
+++ b/Engine/Source/Editor/KismetCompiler/Private/KismetCompiler.cpp
@@ -25,6 +25,10 @@
 #include "Kismet/KismetMathLibrary.h"
 #include "Engine/InheritableComponentHandler.h"
 #include "BlueprintCompilerCppBackendInterface.h"
+#include "Engine/SimpleConstructionScript.h"
+#include "Engine/SCS_Node.h"
+#include "Engine/TimelineTemplate.h"
+#include "Components/TimelineComponent.h"
 
 static bool bDebugPropertyPropagation = false;
 
@@ -1593,10 +1597,10 @@ void FKismetCompilerContext::FinishCompilingFunction(FKismetFunctionContext& Con
 		{
 			if (!Property->HasAnyPropertyFlags(CPF_ZeroConstructor))
 			{
-				Function->FirstPropertyToInit = Property;
-				Function->FunctionFlags |= FUNC_HasDefaults;
+			Function->FirstPropertyToInit = Property;
+			Function->FunctionFlags |= FUNC_HasDefaults;
 				break;
-			}
+		}
 		}
 	}
 
diff --git a/Engine/Source/Editor/KismetCompiler/Private/KismetCompilerMisc.cpp b/Engine/Source/Editor/KismetCompiler/Private/KismetCompilerMisc.cpp
--- a/Engine/Source/Editor/KismetCompiler/Private/KismetCompilerMisc.cpp
+++ b/Engine/Source/Editor/KismetCompiler/Private/KismetCompilerMisc.cpp
@@ -13,6 +13,7 @@
 #include "Editor/UnrealEd/Public/Kismet2/StructureEditorUtils.h"
 #include "Editor/UnrealEd/Public/ObjectTools.h"
 #include "DefaultValueHelper.h"
+#include "Engine/UserDefinedStruct.h"
 
 #define LOCTEXT_NAMESPACE "KismetCompiler"
 
diff --git a/Engine/Source/Editor/KismetCompiler/Private/KismetCompilerModule.cpp b/Engine/Source/Editor/KismetCompiler/Private/KismetCompilerModule.cpp
--- a/Engine/Source/Editor/KismetCompiler/Private/KismetCompilerModule.cpp
+++ b/Engine/Source/Editor/KismetCompiler/Private/KismetCompilerModule.cpp
@@ -13,6 +13,7 @@
 #include "Editor/UnrealEd/Public/Kismet2/KismetReinstanceUtilities.h"
 
 #include "UserDefinedStructureCompilerUtils.h"
+#include "Engine/UserDefinedStruct.h"
 
 DEFINE_LOG_CATEGORY(LogK2Compiler);
 DECLARE_CYCLE_STAT(TEXT("Compile Time"), EKismetCompilerStats_CompileTime, STATGROUP_KismetCompiler);
@@ -93,7 +94,7 @@ void FKismet2CompilerModule::CompileBlueprintInner(class UBlueprint* Blueprint,
 }
 
 
-void FKismet2CompilerModule::CompileStructure(class UUserDefinedStruct* Struct, FCompilerResultsLog& Results)
+void FKismet2CompilerModule::CompileStructure(UUserDefinedStruct* Struct, FCompilerResultsLog& Results)
 {
 	Results.SetSourceName(Struct->GetName());
 	BP_SCOPED_COMPILER_EVENT_STAT(EKismetCompilerStats_CompileTime);
diff --git a/Engine/Source/Editor/KismetCompiler/Private/KismetCompilerVMBackend.cpp b/Engine/Source/Editor/KismetCompiler/Private/KismetCompilerVMBackend.cpp
--- a/Engine/Source/Editor/KismetCompiler/Private/KismetCompilerVMBackend.cpp
+++ b/Engine/Source/Editor/KismetCompiler/Private/KismetCompilerVMBackend.cpp
@@ -13,6 +13,7 @@
 #include "Editor/UnrealEd/Public/Kismet2/StructureEditorUtils.h"
 #include "Editor/UnrealEd/Public/Kismet2/KismetEditorUtilities.h"
 #include "Editor/UnrealEd/Public/Kismet2/KismetDebugUtilities.h"
+#include "Engine/UserDefinedStruct.h"
 
 //////////////////////////////////////////////////////////////////////////
 // FScriptBytecodeWriter
diff --git a/Engine/Source/Editor/KismetCompiler/Private/UserDefinedStructureCompilerUtils.cpp b/Engine/Source/Editor/KismetCompiler/Private/UserDefinedStructureCompilerUtils.cpp
--- a/Engine/Source/Editor/KismetCompiler/Private/UserDefinedStructureCompilerUtils.cpp
+++ b/Engine/Source/Editor/KismetCompiler/Private/UserDefinedStructureCompilerUtils.cpp
@@ -6,6 +6,8 @@
 #include "Editor/UnrealEd/Public/Kismet2/StructureEditorUtils.h"
 #include "Editor/UnrealEd/Public/Kismet2/BlueprintEditorUtils.h"
 #include "Editor/UnrealEd/Public/EditorModes.h"
+#include "Engine/UserDefinedStruct.h"
+#include "UserDefinedStructure/UserDefinedStructEditorData.h"
 
 #define LOCTEXT_NAMESPACE "StructureCompiler"
 
diff --git a/Engine/Source/Editor/KismetCompiler/Public/KismetCompiledFunctionContext.h b/Engine/Source/Editor/KismetCompiler/Public/KismetCompiledFunctionContext.h
--- a/Engine/Source/Editor/KismetCompiler/Public/KismetCompiledFunctionContext.h
+++ b/Engine/Source/Editor/KismetCompiler/Public/KismetCompiledFunctionContext.h
@@ -7,6 +7,8 @@
 #include "BPTerminal.h"
 #include "BlueprintCompiledStatement.h"
 
+class UBlueprintGeneratedClass;
+
 namespace KismetCompilerDebugOptions
 {
 	//@TODO: K2: Turning this off is probably broken due to state merging not working with the current code generation
diff --git a/Engine/Source/Editor/KismetCompiler/Public/KismetCompiler.h b/Engine/Source/Editor/KismetCompiler/Public/KismetCompiler.h
--- a/Engine/Source/Editor/KismetCompiler/Public/KismetCompiler.h
+++ b/Engine/Source/Editor/KismetCompiler/Public/KismetCompiler.h
@@ -7,6 +7,7 @@
 #include "KismetCompiledFunctionContext.h"
 #include "KismetCompilerMisc.h"
 
+class UBlueprintGeneratedClass;
 
 KISMETCOMPILER_API DECLARE_LOG_CATEGORY_EXTERN(LogK2Compiler, Log, All);
 
diff --git a/Engine/Source/Editor/LandscapeEditor/Private/LandscapeEdModeBrushes.cpp b/Engine/Source/Editor/LandscapeEditor/Private/LandscapeEdModeBrushes.cpp
--- a/Engine/Source/Editor/LandscapeEditor/Private/LandscapeEdModeBrushes.cpp
+++ b/Engine/Source/Editor/LandscapeEditor/Private/LandscapeEdModeBrushes.cpp
@@ -12,6 +12,7 @@
 
 #include "LevelUtils.h"
 #include "Materials/MaterialInstanceDynamic.h"
+#include "EngineUtils.h"
 
 // 
 // FLandscapeBrush
diff --git a/Engine/Source/Editor/LevelEditor/Private/SActorDetails.cpp b/Engine/Source/Editor/LevelEditor/Private/SActorDetails.cpp
--- a/Engine/Source/Editor/LevelEditor/Private/SActorDetails.cpp
+++ b/Engine/Source/Editor/LevelEditor/Private/SActorDetails.cpp
@@ -8,6 +8,8 @@
 #include "LevelEditorGenericDetails.h"
 #include "ScopedTransaction.h"
 #include "SourceCodeNavigation.h"
+#include "Engine/Selection.h"
+#include "Engine/BlueprintGeneratedClass.h"
 
 class SActorDetailsUneditableComponentWarning : public SCompoundWidget
 {
diff --git a/Engine/Source/Editor/LevelEditor/Private/SLevelEditor.cpp b/Engine/Source/Editor/LevelEditor/Private/SLevelEditor.cpp
--- a/Engine/Source/Editor/LevelEditor/Private/SLevelEditor.cpp
+++ b/Engine/Source/Editor/LevelEditor/Private/SLevelEditor.cpp
@@ -37,6 +37,7 @@
 #include "SDockTab.h"
 #include "SActorDetails.h"
 #include "ScopedTransaction.h"
+#include "GameFramework/WorldSettings.h"
 
 
 static const FName LevelEditorBuildAndSubmitTab("LevelEditorBuildAndSubmit");
diff --git a/Engine/Source/Editor/LevelEditor/Private/SLevelViewportToolBar.cpp b/Engine/Source/Editor/LevelEditor/Private/SLevelViewportToolBar.cpp
--- a/Engine/Source/Editor/LevelEditor/Private/SLevelViewportToolBar.cpp
+++ b/Engine/Source/Editor/LevelEditor/Private/SLevelViewportToolBar.cpp
@@ -21,6 +21,7 @@
 #include "Editor/UnrealEd/Public/SEditorViewportViewMenu.h"
 #include "StatsData.h"
 #include "BufferVisualizationData.h"
+#include "GameFramework/WorldSettings.h"
 
 #define LOCTEXT_NAMESPACE "LevelViewportToolBar"
 
diff --git a/Engine/Source/Editor/MaterialEditor/Private/MaterialEditor.cpp b/Engine/Source/Editor/MaterialEditor/Private/MaterialEditor.cpp
--- a/Engine/Source/Editor/MaterialEditor/Private/MaterialEditor.cpp
+++ b/Engine/Source/Editor/MaterialEditor/Private/MaterialEditor.cpp
@@ -29,6 +29,7 @@
 #include "Materials/MaterialExpressionTransformPosition.h"
 #include "Materials/MaterialExpressionVectorParameter.h"
 #include "Materials/MaterialFunction.h"
+#include "Materials/MaterialParameterCollection.h"
 
 #include "MaterialEditorActions.h"
 #include "MaterialExpressionClasses.h"
@@ -68,6 +69,9 @@
 #include "Developer/MessageLog/Public/MessageLogModule.h"
 #include "Particles/ParticleSystemComponent.h"
 #include "GenericCommands.h"
+#include "CanvasTypes.h"
+#include "Engine/Selection.h"
+#include "Engine/TextureCube.h"
 
 #define LOCTEXT_NAMESPACE "MaterialEditor"
 
diff --git a/Engine/Source/Editor/Matinee/Private/MatineePropertyWindow.cpp b/Engine/Source/Editor/Matinee/Private/MatineePropertyWindow.cpp
--- a/Engine/Source/Editor/Matinee/Private/MatineePropertyWindow.cpp
+++ b/Engine/Source/Editor/Matinee/Private/MatineePropertyWindow.cpp
@@ -5,6 +5,7 @@
 #include "Matinee/MatineeActorCameraAnim.h"
 #include "Editor/PropertyEditor/Public/PropertyEditorModule.h"
 #include "Editor/PropertyEditor/Public/IDetailsView.h"
+#include "Camera/CameraAnim.h"
 
 void FMatinee::BuildPropertyWindow()
 {
diff --git a/Engine/Source/Editor/NiagaraEditor/Private/SNiagaraEffectEditorWidget.h b/Engine/Source/Editor/NiagaraEditor/Private/SNiagaraEffectEditorWidget.h
--- a/Engine/Source/Editor/NiagaraEditor/Private/SNiagaraEffectEditorWidget.h
+++ b/Engine/Source/Editor/NiagaraEditor/Private/SNiagaraEffectEditorWidget.h
@@ -17,6 +17,7 @@
 #include "Engine/NiagaraEffectRenderer.h"
 #include "ComponentReregisterContext.h"
 #include "SNumericEntryBox.h"
+#include "Curves/CurveVector.h"
 
 
 #define NGED_SECTION_BORDER SNew(SBorder).BorderImage(FEditorStyle::GetBrush("ToolPanel.GroupBorder")).Padding(2.0f).HAlign(HAlign_Left)
diff --git a/Engine/Source/Editor/Persona/Private/AnimationEditorPreviewScene.cpp b/Engine/Source/Editor/Persona/Private/AnimationEditorPreviewScene.cpp
--- a/Engine/Source/Editor/Persona/Private/AnimationEditorPreviewScene.cpp
+++ b/Engine/Source/Editor/Persona/Private/AnimationEditorPreviewScene.cpp
@@ -3,6 +3,9 @@
 #include "PersonaPrivatePCH.h"
 
 #include "AnimationEditorPreviewScene.h"
+#include "Components/DirectionalLightComponent.h"
+#include "Components/ExponentialHeightFogComponent.h"
+#include "Components/SphereReflectionCaptureComponent.h"
 
 /////////////////////////////////////////////////////////////////////////
 // FAnimationEditorPreviewScene
diff --git a/Engine/Source/Editor/Persona/Private/Persona.cpp b/Engine/Source/Editor/Persona/Private/Persona.cpp
--- a/Engine/Source/Editor/Persona/Private/Persona.cpp
+++ b/Engine/Source/Editor/Persona/Private/Persona.cpp
@@ -69,6 +69,9 @@
 #include "NotificationManager.h"
 
 #include "Editor/KismetWidgets/Public/SSingleObjectDetailsPanel.h"
+#include "Animation/AimOffsetBlendSpace.h"
+#include "Animation/AimOffsetBlendSpace1D.h"
+#include "Animation/AnimNotifies/AnimNotifyState.h"
 
 #define LOCTEXT_NAMESPACE "FPersona"
 
diff --git a/Engine/Source/Editor/Persona/Private/SAnimCompositeEditor.h b/Engine/Source/Editor/Persona/Private/SAnimCompositeEditor.h
--- a/Engine/Source/Editor/Persona/Private/SAnimCompositeEditor.h
+++ b/Engine/Source/Editor/Persona/Private/SAnimCompositeEditor.h
@@ -5,6 +5,9 @@
 
 #include "Persona.h"
 #include "SAnimEditorBase.h"
+#include "Animation/AnimComposite.h"
+
+class UAnimComposite;
 
 //////////////////////////////////////////////////////////////////////////
 // SAnimCompositeEditor
diff --git a/Engine/Source/Editor/Persona/Private/SAnimNotifyPanel.cpp b/Engine/Source/Editor/Persona/Private/SAnimNotifyPanel.cpp
--- a/Engine/Source/Editor/Persona/Private/SAnimNotifyPanel.cpp
+++ b/Engine/Source/Editor/Persona/Private/SAnimNotifyPanel.cpp
@@ -14,6 +14,8 @@
 #include "SExpandableArea.h"
 #include "Toolkits/AssetEditorManager.h"
 #include "BlueprintActionDatabase.h"
+#include "Animation/AnimNotifies/AnimNotifyState.h"
+#include "Animation/AnimNotifies/AnimNotify.h"
 
 // Track Panel drawing
 const float NotificationTrackHeight = 20.0f;
diff --git a/Engine/Source/Editor/Persona/Private/SAnimationScrubPanel.h b/Engine/Source/Editor/Persona/Private/SAnimationScrubPanel.h
--- a/Engine/Source/Editor/Persona/Private/SAnimationScrubPanel.h
+++ b/Engine/Source/Editor/Persona/Private/SAnimationScrubPanel.h
@@ -6,6 +6,8 @@
 #include "SScrubWidget.h"
 #include "Editor/EditorWidgets/Public/ITransportControl.h"
 
+struct FAnimBlueprintDebugData;
+
 class SAnimationScrubPanel : public SCompoundWidget
 {
 public:
diff --git a/Engine/Source/Editor/ProjectSettingsViewer/Private/ProjectSettingsViewerModule.cpp b/Engine/Source/Editor/ProjectSettingsViewer/Private/ProjectSettingsViewerModule.cpp
--- a/Engine/Source/Editor/ProjectSettingsViewer/Private/ProjectSettingsViewerModule.cpp
+++ b/Engine/Source/Editor/ProjectSettingsViewer/Private/ProjectSettingsViewerModule.cpp
@@ -22,6 +22,10 @@
 #include "Editor/UnrealEd/Public/Settings/EditorProjectSettings.h"
 
 #include "AISystem.h"
+#include "GameFramework/InputSettings.h"
+#include "AI/Navigation/RecastNavMesh.h"
+#include "Engine/NetworkSettings.h"
+#include "PhysicsEngine/PhysicsSettings.h"
 
 #define LOCTEXT_NAMESPACE "FProjectSettingsViewerModule"
 
diff --git a/Engine/Source/Editor/ReferenceViewer/Private/SReferenceViewer.cpp b/Engine/Source/Editor/ReferenceViewer/Private/SReferenceViewer.cpp
--- a/Engine/Source/Editor/ReferenceViewer/Private/SReferenceViewer.cpp
+++ b/Engine/Source/Editor/ReferenceViewer/Private/SReferenceViewer.cpp
@@ -10,6 +10,7 @@
 #include "GlobalEditorCommonCommands.h"
 
 #include "Editor/UnrealEd/Public/ObjectTools.h"
+#include "Engine/Selection.h"
 
 #define LOCTEXT_NAMESPACE "ReferenceViewer"
 
diff --git a/Engine/Source/Editor/SoundCueEditor/Private/SoundCueEditor.cpp b/Engine/Source/Editor/SoundCueEditor/Private/SoundCueEditor.cpp
--- a/Engine/Source/Editor/SoundCueEditor/Private/SoundCueEditor.cpp
+++ b/Engine/Source/Editor/SoundCueEditor/Private/SoundCueEditor.cpp
@@ -18,6 +18,7 @@
 #include "SSoundCuePalette.h"
 #include "SDockTab.h"
 #include "GenericCommands.h"
+#include "Sound/SoundCue.h"
 
 #define LOCTEXT_NAMESPACE "SoundCueEditor"
 
diff --git a/Engine/Source/Editor/SoundCueEditor/Private/SoundCueEditorUtilities.cpp b/Engine/Source/Editor/SoundCueEditor/Private/SoundCueEditorUtilities.cpp
--- a/Engine/Source/Editor/SoundCueEditor/Private/SoundCueEditorUtilities.cpp
+++ b/Engine/Source/Editor/SoundCueEditor/Private/SoundCueEditorUtilities.cpp
@@ -7,6 +7,7 @@
 #include "SoundDefinitions.h"
 #include "Sound/SoundNodeWavePlayer.h"
 #include "GraphEditor.h"
+#include "Sound/SoundCue.h"
 
 bool FSoundCueEditorUtilities::CanPasteNodes(const class UEdGraph* Graph)
 {
diff --git a/Engine/Source/Editor/StaticMeshEditor/Private/StaticMeshEditorViewportClient.cpp b/Engine/Source/Editor/StaticMeshEditor/Private/StaticMeshEditorViewportClient.cpp
--- a/Engine/Source/Editor/StaticMeshEditor/Private/StaticMeshEditorViewportClient.cpp
+++ b/Engine/Source/Editor/StaticMeshEditor/Private/StaticMeshEditorViewportClient.cpp
@@ -19,6 +19,7 @@
 
 #include "Runtime/Analytics/Analytics/Public/Interfaces/IAnalyticsProvider.h"
 #include "EngineAnalytics.h"
+#include "AI/Navigation/NavCollision.h"
 
 #if WITH_PHYSX
 #include "Editor/UnrealEd/Private/EditorPhysXSupport.h"
diff --git a/Engine/Source/Editor/StatsViewer/Private/StatsEntries/LightingBuildInfo.cpp b/Engine/Source/Editor/StatsViewer/Private/StatsEntries/LightingBuildInfo.cpp
--- a/Engine/Source/Editor/StatsViewer/Private/StatsEntries/LightingBuildInfo.cpp
+++ b/Engine/Source/Editor/StatsViewer/Private/StatsEntries/LightingBuildInfo.cpp
@@ -2,6 +2,7 @@
 
 #include "StatsViewerPrivatePCH.h"
 #include "LightingBuildInfo.h"
+#include "Lightmass/LightmappedSurfaceCollection.h"
 
 ULightingBuildInfo::ULightingBuildInfo(const FObjectInitializer& ObjectInitializer)
 	: Super(ObjectInitializer)
diff --git a/Engine/Source/Editor/StatsViewer/Private/StatsEntries/StaticMeshLightingInfo.cpp b/Engine/Source/Editor/StatsViewer/Private/StatsEntries/StaticMeshLightingInfo.cpp
--- a/Engine/Source/Editor/StatsViewer/Private/StatsEntries/StaticMeshLightingInfo.cpp
+++ b/Engine/Source/Editor/StatsViewer/Private/StatsEntries/StaticMeshLightingInfo.cpp
@@ -2,6 +2,7 @@
 
 #include "StatsViewerPrivatePCH.h"
 #include "StaticMeshLightingInfo.h"
+#include "Lightmass/LightmappedSurfaceCollection.h"
 
 #define LOCTEXT_NAMESPACE "Editor.StatsViewer"
 
diff --git a/Engine/Source/Editor/StatsViewer/Private/StatsPages/PrimitiveStatsPage.cpp b/Engine/Source/Editor/StatsViewer/Private/StatsPages/PrimitiveStatsPage.cpp
--- a/Engine/Source/Editor/StatsViewer/Private/StatsPages/PrimitiveStatsPage.cpp
+++ b/Engine/Source/Editor/StatsViewer/Private/StatsPages/PrimitiveStatsPage.cpp
@@ -8,6 +8,8 @@
 #include "LandscapeComponent.h"
 #include "LightMap.h"
 #include "ShadowMap.h"
+#include "Components/ModelComponent.h"
+#include "Engine/Selection.h"
 
 #define LOCTEXT_NAMESPACE "Editor.StatsViewer.PrimitiveStats"
 
diff --git a/Engine/Source/Editor/StatsViewer/Private/StatsPages/StaticMeshLightingInfoStatsPage.cpp b/Engine/Source/Editor/StatsViewer/Private/StatsPages/StaticMeshLightingInfoStatsPage.cpp
--- a/Engine/Source/Editor/StatsViewer/Private/StatsPages/StaticMeshLightingInfoStatsPage.cpp
+++ b/Engine/Source/Editor/StatsViewer/Private/StatsPages/StaticMeshLightingInfoStatsPage.cpp
@@ -8,6 +8,7 @@
 #include "BusyCursor.h"
 #include "ScopedTransaction.h"
 #include "STextEntryPopup.h"
+#include "Engine/LevelStreaming.h"
 
 #define LOCTEXT_NAMESPACE "Editor.StatsViewer.StaticMeshLightingInfo"
 
diff --git a/Engine/Source/Editor/StatsViewer/Private/StatsPages/TextureStatsPage.cpp b/Engine/Source/Editor/StatsViewer/Private/StatsPages/TextureStatsPage.cpp
--- a/Engine/Source/Editor/StatsViewer/Private/StatsPages/TextureStatsPage.cpp
+++ b/Engine/Source/Editor/StatsViewer/Private/StatsPages/TextureStatsPage.cpp
@@ -5,6 +5,8 @@
 #include "TextureStats.h"
 #include "ReferencedAssetsUtils.h"
 #include "AssetSelection.h"
+#include "EngineUtils.h"
+#include "Engine/TextureCube.h"
 
 #define LOCTEXT_NAMESPACE "Editor.StatsViewer.TextureStats"
 
diff --git a/Engine/Source/Editor/TextureEditor/Private/TextureEditorToolkit.cpp b/Engine/Source/Editor/TextureEditor/Private/TextureEditorToolkit.cpp
--- a/Engine/Source/Editor/TextureEditor/Private/TextureEditorToolkit.cpp
+++ b/Engine/Source/Editor/TextureEditor/Private/TextureEditorToolkit.cpp
@@ -14,6 +14,7 @@
 #include "Engine/LightMapTexture2D.h"
 #include "Engine/ShadowMapTexture2D.h"
 #include "Engine/Texture2DDynamic.h"
+#include "DeviceProfiles/DeviceProfileManager.h"
 
 #define LOCTEXT_NAMESPACE "FTextureEditorToolkit"
 
diff --git a/Engine/Source/Editor/UMGEditor/Private/WidgetBlueprint.cpp b/Engine/Source/Editor/UMGEditor/Private/WidgetBlueprint.cpp
--- a/Engine/Source/Editor/UMGEditor/Private/WidgetBlueprint.cpp
+++ b/Engine/Source/Editor/UMGEditor/Private/WidgetBlueprint.cpp
@@ -11,6 +11,7 @@
 #include "WidgetBlueprint.h"
 #include "WidgetBlueprintCompiler.h"
 #include "PropertyBinding.h"
+#include "Engine/UserDefinedStruct.h"
 
 #define LOCTEXT_NAMESPACE "UMG"
 
diff --git a/Engine/Source/Editor/UnrealEd/Private/AssetEditorModeManager.cpp b/Engine/Source/Editor/UnrealEd/Private/AssetEditorModeManager.cpp
--- a/Engine/Source/Editor/UnrealEd/Private/AssetEditorModeManager.cpp
+++ b/Engine/Source/Editor/UnrealEd/Private/AssetEditorModeManager.cpp
@@ -3,6 +3,7 @@
 #include "UnrealEd.h"
 #include "AssetEditorModeManager.h"
 #include "PreviewScene.h"
+#include "Engine/Selection.h"
 
 //////////////////////////////////////////////////////////////////////////
 // FAssetEditorModeManager
diff --git a/Engine/Source/Editor/UnrealEd/Private/Commandlets/GatherTextFromAssetsCommandlet.cpp b/Engine/Source/Editor/UnrealEd/Private/Commandlets/GatherTextFromAssetsCommandlet.cpp
--- a/Engine/Source/Editor/UnrealEd/Private/Commandlets/GatherTextFromAssetsCommandlet.cpp
+++ b/Engine/Source/Editor/UnrealEd/Private/Commandlets/GatherTextFromAssetsCommandlet.cpp
@@ -9,6 +9,7 @@
 #include "InternationalizationMetadata.h"
 #include "Sound/DialogueWave.h"
 #include "Sound/DialogueVoice.h"
+#include "Engine/DataTable.h"
 
 DEFINE_LOG_CATEGORY_STATIC(LogGatherTextFromAssetsCommandlet, Log, All);
 
diff --git a/Engine/Source/Editor/UnrealEd/Private/DataTableEditorUtils.cpp b/Engine/Source/Editor/UnrealEd/Private/DataTableEditorUtils.cpp
--- a/Engine/Source/Editor/UnrealEd/Private/DataTableEditorUtils.cpp
+++ b/Engine/Source/Editor/UnrealEd/Private/DataTableEditorUtils.cpp
@@ -3,6 +3,7 @@
 #include "UnrealEd.h"
 #include "DataTableEditorUtils.h"
 #include "K2Node_GetDataTableRow.h"
+#include "Engine/UserDefinedStruct.h"
 
 FDataTableEditorUtils::FDataTableEditorManager& FDataTableEditorUtils::FDataTableEditorManager::Get()
 {
diff --git a/Engine/Source/Editor/UnrealEd/Private/Dialogs/SMeshMergingDialog.cpp b/Engine/Source/Editor/UnrealEd/Private/Dialogs/SMeshMergingDialog.cpp
--- a/Engine/Source/Editor/UnrealEd/Private/Dialogs/SMeshMergingDialog.cpp
+++ b/Engine/Source/Editor/UnrealEd/Private/Dialogs/SMeshMergingDialog.cpp
@@ -13,6 +13,7 @@
 #include "Engine/StaticMesh.h"
 #include "Engine/Selection.h"
 #include "SystemSettings.h"
+#include "Engine/TextureLODSettings.h"
 #include "ScopedTransaction.h"
 
 #define LOCTEXT_NAMESPACE "SMeshMergingDialog"
diff --git a/Engine/Source/Editor/UnrealEd/Private/Editor.cpp b/Engine/Source/Editor/UnrealEd/Private/Editor.cpp
--- a/Engine/Source/Editor/UnrealEd/Private/Editor.cpp
+++ b/Engine/Source/Editor/UnrealEd/Private/Editor.cpp
@@ -97,6 +97,18 @@
 #include "NotificationManager.h"
 #include "Engine/GameEngine.h"
 #include "Engine/TextureRenderTarget2D.h"
+#include "GameFramework/GameUserSettings.h"
+#include "Engine/Light.h"
+#include "Engine/LevelStreamingVolume.h"
+#include "Sound/SoundCue.h"
+#include "Components/BrushComponent.h"
+#include "Engine/LocalPlayer.h"
+#include "Components/SkyLightComponent.h"
+#include "Components/ReflectionCaptureComponent.h"
+#include "Engine/Polys.h"
+#include "UnrealEngine.h"
+#include "EngineStats.h"
+#include "Engine/SimpleConstructionScript.h"
 
 DEFINE_LOG_CATEGORY_STATIC(LogEditor, Log, All);
 
diff --git a/Engine/Source/Editor/UnrealEd/Private/EditorExporters.cpp b/Engine/Source/Editor/UnrealEd/Private/EditorExporters.cpp
--- a/Engine/Source/Editor/UnrealEd/Private/EditorExporters.cpp
+++ b/Engine/Source/Editor/UnrealEd/Private/EditorExporters.cpp
@@ -38,6 +38,7 @@
 #include "EngineUtils.h"
 #include "DeviceProfiles/DeviceProfileManager.h"
 #include "TileRendering.h"
+#include "Engine/TextureLODSettings.h"
 
 DEFINE_LOG_CATEGORY_STATIC(LogEditorExporters, Log, All);
 
@@ -2460,7 +2461,7 @@ namespace MaterialExportUtils
 		float CurrentRealTime = 0.f;
 		float CurrentWorldTime = 0.f;
 		float DeltaWorldTime = 0.f;
-
+	
 		if (InWorld)
 		{
 			CurrentRealTime = InWorld->GetRealTimeSeconds();
@@ -2546,15 +2547,15 @@ namespace MaterialExportUtils
 			OutBMP.Add(UniformValue);
 			return true;
 		}
-
+	
 		check(InRenderTarget);
-		
+
 		RenderMaterialTile(InWorld, MaterialProxy, InRenderTarget);
-		
+				
 		bool bNormalmap = (InMaterialProperty == MP_Normal);
 		FReadSurfaceDataFlags ReadPixelFlags(bNormalmap ? RCM_SNorm : RCM_UNorm);
 		ReadPixelFlags.SetLinearToGamma(false);
-		
+
 		FTextureRenderTargetResource* RTResource = InRenderTarget->GameThread_GetRenderTargetResource();
 		return RTResource->ReadPixels(OutBMP, ReadPixelFlags);
 	}
diff --git a/Engine/Source/Editor/UnrealEd/Private/EditorObject.cpp b/Engine/Source/Editor/UnrealEd/Private/EditorObject.cpp
--- a/Engine/Source/Editor/UnrealEd/Private/EditorObject.cpp
+++ b/Engine/Source/Editor/UnrealEd/Private/EditorObject.cpp
@@ -14,6 +14,7 @@
 #include "InstancedFoliage.h"
 #include "Components/BrushComponent.h"
 #include "Engine/StaticMesh.h"
+#include "Engine/Polys.h"
 
 DEFINE_LOG_CATEGORY_STATIC(LogEditorObject, Log, All);
 
diff --git a/Engine/Source/Editor/UnrealEd/Private/EditorSelectUtils.cpp b/Engine/Source/Editor/UnrealEd/Private/EditorSelectUtils.cpp
--- a/Engine/Source/Editor/UnrealEd/Private/EditorSelectUtils.cpp
+++ b/Engine/Source/Editor/UnrealEd/Private/EditorSelectUtils.cpp
@@ -8,6 +8,8 @@
 #include "SnappingUtils.h"
 #include "MessageLog.h"
 #include "ComponentEditorUtils.h"
+#include "Components/ChildActorComponent.h"
+#include "Components/DecalComponent.h"
 
 #define LOCTEXT_NAMESPACE "EditorSelectUtils"
 
diff --git a/Engine/Source/Editor/UnrealEd/Private/EditorServer.cpp b/Engine/Source/Editor/UnrealEd/Private/EditorServer.cpp
--- a/Engine/Source/Editor/UnrealEd/Private/EditorServer.cpp
+++ b/Engine/Source/Editor/UnrealEd/Private/EditorServer.cpp
@@ -60,6 +60,13 @@
 #include "NotificationManager.h"
 #include "EditorUndoClient.h"
 #include "DesktopPlatformModule.h"
+#include "Layers/Layer.h"
+#include "Engine/Light.h"
+#include "Animation/AnimNotifies/AnimNotify.h"
+#include "Components/BillboardComponent.h"
+#include "Components/DrawFrustumComponent.h"
+#include "UnrealEngine.h"
+#include "AI/Navigation/NavLinkRenderingComponent.h"
 
 DEFINE_LOG_CATEGORY_STATIC(LogEditorServer, Log, All);
 
diff --git a/Engine/Source/Editor/UnrealEd/Private/EditorViewportClient.cpp b/Engine/Source/Editor/UnrealEd/Private/EditorViewportClient.cpp
--- a/Engine/Source/Editor/UnrealEd/Private/EditorViewportClient.cpp
+++ b/Engine/Source/Editor/UnrealEd/Private/EditorViewportClient.cpp
@@ -26,6 +26,7 @@
 #include "EngineUtils.h"
 #include "SEditorViewport.h"
 #include "AssetEditorModeManager.h"
+#include "Components/DirectionalLightComponent.h"
 
 #define LOCTEXT_NAMESPACE "EditorViewportClient"
 
diff --git a/Engine/Source/Editor/UnrealEd/Private/Factories/ActorFactory.cpp b/Engine/Source/Editor/UnrealEd/Private/Factories/ActorFactory.cpp
--- a/Engine/Source/Editor/UnrealEd/Private/Factories/ActorFactory.cpp
+++ b/Engine/Source/Editor/UnrealEd/Private/Factories/ActorFactory.cpp
@@ -65,6 +65,7 @@ ActorFactory.cpp:
 #include "Components/VectorFieldComponent.h"
 #include "Animation/AnimBlueprintGeneratedClass.h"
 #include "Kismet2/ComponentEditorUtils.h"
+#include "Components/BillboardComponent.h"
 
 DEFINE_LOG_CATEGORY_STATIC(LogActorFactory, Log, All);
 
diff --git a/Engine/Source/Editor/UnrealEd/Private/Fbx/FbxLightImport.cpp b/Engine/Source/Editor/UnrealEd/Private/Fbx/FbxLightImport.cpp
--- a/Engine/Source/Editor/UnrealEd/Private/Fbx/FbxLightImport.cpp
+++ b/Engine/Source/Editor/UnrealEd/Private/Fbx/FbxLightImport.cpp
@@ -8,6 +8,9 @@
 #include "Factories.h"
 #include "Engine.h"
 #include "FbxImporter.h"
+#include "Engine/PointLight.h"
+#include "Engine/DirectionalLight.h"
+#include "Engine/SpotLight.h"
 
 DEFINE_LOG_CATEGORY_STATIC(LogFbxLightImport, Log, All);
 
diff --git a/Engine/Source/Editor/UnrealEd/Private/Fbx/FbxMainExport.cpp b/Engine/Source/Editor/UnrealEd/Private/Fbx/FbxMainExport.cpp
--- a/Engine/Source/Editor/UnrealEd/Private/Fbx/FbxMainExport.cpp
+++ b/Engine/Source/Editor/UnrealEd/Private/Fbx/FbxMainExport.cpp
@@ -67,6 +67,7 @@
 #include "Components/PointLightComponent.h"
 #include "Camera/CameraActor.h"
 #include "Components/DirectionalLightComponent.h"
+#include "Components/ChildActorComponent.h"
 
 namespace UnFbx
 {
diff --git a/Engine/Source/Editor/UnrealEd/Private/Fbx/FbxSkeletalMeshImport.cpp b/Engine/Source/Editor/UnrealEd/Private/Fbx/FbxSkeletalMeshImport.cpp
--- a/Engine/Source/Editor/UnrealEd/Private/Fbx/FbxSkeletalMeshImport.cpp
+++ b/Engine/Source/Editor/UnrealEd/Private/Fbx/FbxSkeletalMeshImport.cpp
@@ -54,6 +54,8 @@
 #include "ComponentReregisterContext.h"
 
 #include "FbxErrors.h"
+#include "PhysicsEngine/PhysicsAsset.h"
+#include "Engine/SkeletalMeshSocket.h"
 
 #define LOCTEXT_NAMESPACE "FBXImpoter"
 
diff --git a/Engine/Source/Editor/UnrealEd/Private/Fbx/FbxStaticMeshImport.cpp b/Engine/Source/Editor/UnrealEd/Private/Fbx/FbxStaticMeshImport.cpp
--- a/Engine/Source/Editor/UnrealEd/Private/Fbx/FbxStaticMeshImport.cpp
+++ b/Engine/Source/Editor/UnrealEd/Private/Fbx/FbxStaticMeshImport.cpp
@@ -46,6 +46,8 @@
 #include "../Private/GeomFitUtils.h"
 #include "TargetPlatform.h"
 #include "FbxErrors.h"
+#include "Engine/StaticMeshSocket.h"
+#include "Engine/Polys.h"
 
 #define LOCTEXT_NAMESPACE "FbxStaticMeshImport"
 
diff --git a/Engine/Source/Editor/UnrealEd/Private/HierarchicalLOD.cpp b/Engine/Source/Editor/UnrealEd/Private/HierarchicalLOD.cpp
--- a/Engine/Source/Editor/UnrealEd/Private/HierarchicalLOD.cpp
+++ b/Engine/Source/Editor/UnrealEd/Private/HierarchicalLOD.cpp
@@ -16,6 +16,8 @@
 #include "ObjectTools.h"
 
 #endif // WITH_EDITOR
+#include "GameFramework/WorldSettings.h"
+#include "Components/InstancedStaticMeshComponent.h"
 
 DEFINE_LOG_CATEGORY_STATIC(LogLODGenerator, Log, All);
 
diff --git a/Engine/Source/Editor/UnrealEd/Private/Kismet2/BlueprintEditorUtils.cpp b/Engine/Source/Editor/UnrealEd/Private/Kismet2/BlueprintEditorUtils.cpp
--- a/Engine/Source/Editor/UnrealEd/Private/Kismet2/BlueprintEditorUtils.cpp
+++ b/Engine/Source/Editor/UnrealEd/Private/Kismet2/BlueprintEditorUtils.cpp
@@ -51,6 +51,8 @@
 #include "Engine/InheritableComponentHandler.h"
 
 #include "EditorCategoryUtils.h"
+#include "EngineUtils.h"
+#include "Engine/LevelScriptActor.h"
 #define LOCTEXT_NAMESPACE "Blueprint"
 
 DEFINE_LOG_CATEGORY(LogBlueprintDebug);
diff --git a/Engine/Source/Editor/UnrealEd/Private/Kismet2/ComponentEditorUtils.cpp b/Engine/Source/Editor/UnrealEd/Private/Kismet2/ComponentEditorUtils.cpp
--- a/Engine/Source/Editor/UnrealEd/Private/Kismet2/ComponentEditorUtils.cpp
+++ b/Engine/Source/Editor/UnrealEd/Private/Kismet2/ComponentEditorUtils.cpp
@@ -14,6 +14,7 @@
 #include "SourceCodeNavigation.h"
 #include "ClassIconFinder.h"
 #include "AssetEditorManager.h"
+#include "Components/DecalComponent.h"
 
 #define LOCTEXT_NAMESPACE "ComponentEditorUtils"
 
diff --git a/Engine/Source/Editor/UnrealEd/Private/Kismet2/Kismet2.cpp b/Engine/Source/Editor/UnrealEd/Private/Kismet2/Kismet2.cpp
--- a/Engine/Source/Editor/UnrealEd/Private/Kismet2/Kismet2.cpp
+++ b/Engine/Source/Editor/UnrealEd/Private/Kismet2/Kismet2.cpp
@@ -37,6 +37,10 @@
 #include "Editor/UnrealEd/Public/PackageTools.h"
 #include "NotificationManager.h"
 #include "SNotificationList.h" // for FNotificationInfo
+#include "Engine/BlueprintGeneratedClass.h"
+#include "Engine/SimpleConstructionScript.h"
+#include "Engine/SCS_Node.h"
+#include "GeneralProjectSettings.h"
 
 DECLARE_CYCLE_STAT(TEXT("Compile Blueprint"), EKismetCompilerStats_CompileBlueprint, STATGROUP_KismetCompiler);
 DECLARE_CYCLE_STAT(TEXT("Broadcast Precompile"), EKismetCompilerStats_BroadcastPrecompile, STATGROUP_KismetCompiler);
diff --git a/Engine/Source/Editor/UnrealEd/Private/Kismet2/KismetReinstanceUtilities.cpp b/Engine/Source/Editor/UnrealEd/Private/Kismet2/KismetReinstanceUtilities.cpp
--- a/Engine/Source/Editor/UnrealEd/Private/Kismet2/KismetReinstanceUtilities.cpp
+++ b/Engine/Source/Editor/UnrealEd/Private/Kismet2/KismetReinstanceUtilities.cpp
@@ -9,6 +9,8 @@
 #include "ComponentInstanceDataCache.h"
 #include "Engine/DynamicBlueprintBinding.h"
 #include "BlueprintEditor.h"
+#include "Engine/BlueprintGeneratedClass.h"
+#include "Engine/Selection.h"
 
 DECLARE_CYCLE_STAT(TEXT("Replace Instances"), EKismetReinstancerStats_ReplaceInstancesOfClass, STATGROUP_KismetReinstancer );
 DECLARE_CYCLE_STAT(TEXT("Find Referencers"), EKismetReinstancerStats_FindReferencers, STATGROUP_KismetReinstancer );
diff --git a/Engine/Source/Editor/UnrealEd/Private/Kismet2/StructureEditorUtils.cpp b/Engine/Source/Editor/UnrealEd/Private/Kismet2/StructureEditorUtils.cpp
--- a/Engine/Source/Editor/UnrealEd/Private/Kismet2/StructureEditorUtils.cpp
+++ b/Engine/Source/Editor/UnrealEd/Private/Kismet2/StructureEditorUtils.cpp
@@ -13,6 +13,7 @@
 #include "Toolkits/AssetEditorManager.h"
 #include "Editor/DataTableEditor/Public/IDataTableEditor.h"
 #include "Engine/UserDefinedStruct.h"
+#include "Engine/DataTable.h"
 
 #define LOCTEXT_NAMESPACE "Structure"
 
diff --git a/Engine/Source/Editor/UnrealEd/Private/LevelViewportClickHandlers.cpp b/Engine/Source/Editor/UnrealEd/Private/LevelViewportClickHandlers.cpp
--- a/Engine/Source/Editor/UnrealEd/Private/LevelViewportClickHandlers.cpp
+++ b/Engine/Source/Editor/UnrealEd/Private/LevelViewportClickHandlers.cpp
@@ -17,6 +17,7 @@
 #include "Engine/Selection.h"
 #include "Engine/StaticMeshActor.h"
 #include "Engine/TargetPoint.h"
+#include "EngineUtils.h"
 
 #define LOCTEXT_NAMESPACE "ClickHandlers"
 
diff --git a/Engine/Source/Editor/UnrealEd/Private/LevelViewportClickHandlers.h b/Engine/Source/Editor/UnrealEd/Private/LevelViewportClickHandlers.h
--- a/Engine/Source/Editor/UnrealEd/Private/LevelViewportClickHandlers.h
+++ b/Engine/Source/Editor/UnrealEd/Private/LevelViewportClickHandlers.h
@@ -7,6 +7,7 @@ struct FViewportClick;
 struct HGeomEdgeProxy;
 struct HGeomPolyProxy;
 struct HGeomVertexProxy;
+struct HActor;
 
 /**
  * A hit proxy class for sockets in the main editor viewports.
diff --git a/Engine/Source/Editor/UnrealEd/Private/Lightmass/Lightmass.cpp b/Engine/Source/Editor/UnrealEd/Private/Lightmass/Lightmass.cpp
--- a/Engine/Source/Editor/UnrealEd/Private/Lightmass/Lightmass.cpp
+++ b/Engine/Source/Editor/UnrealEd/Private/Lightmass/Lightmass.cpp
@@ -20,6 +20,12 @@
 #include "ComponentReregisterContext.h"
 #include "ShaderCompiler.h"
 #include "Lightmass/Lightmass.h"
+#include "Engine/LevelStreaming.h"
+#include "Components/DirectionalLightComponent.h"
+#include "Components/ModelComponent.h"
+#include "Engine/GeneratedMeshAreaLight.h"
+#include "Components/SkyLightComponent.h"
+#include "UnrealEngine.h"
 
 extern FSwarmDebugOptions GSwarmDebugOptions;
 
diff --git a/Engine/Source/Editor/UnrealEd/Private/MaterialEditor/MaterialEditorMeshComponent.cpp b/Engine/Source/Editor/UnrealEd/Private/MaterialEditor/MaterialEditorMeshComponent.cpp
--- a/Engine/Source/Editor/UnrealEd/Private/MaterialEditor/MaterialEditorMeshComponent.cpp
+++ b/Engine/Source/Editor/UnrealEd/Private/MaterialEditor/MaterialEditorMeshComponent.cpp
@@ -1,6 +1,8 @@
 // Copyright 1998-2015 Epic Games, Inc. All Rights Reserved.
 
 #include "UnrealEd.h"
+#include "MaterialEditor/MaterialEditorMeshComponent.h"
+#include "PhysicsEngine/BodySetup.h"
 
 FBoxSphereBounds UMaterialEditorMeshComponent::CalcBounds(const FTransform& LocalToWorld) const
 {
diff --git a/Engine/Source/Editor/UnrealEd/Private/ObjectTools.cpp b/Engine/Source/Editor/UnrealEd/Private/ObjectTools.cpp
--- a/Engine/Source/Editor/UnrealEd/Private/ObjectTools.cpp
+++ b/Engine/Source/Editor/UnrealEd/Private/ObjectTools.cpp
@@ -33,6 +33,14 @@
 #include "SNotificationList.h"
 #include "NotificationManager.h"
 #include "Layers/ILayers.h"
+#include "PhysicalMaterials/PhysicalMaterial.h"
+#include "Engine/BlueprintGeneratedClass.h"
+#include "Engine/SimpleConstructionScript.h"
+#include "Engine/LevelStreaming.h"
+#include "GameFramework/WorldSettings.h"
+#include "CanvasTypes.h"
+#include "Engine/SCS_Node.h"
+#include "Engine/UserDefinedStruct.h"
 
 DEFINE_LOG_CATEGORY_STATIC(LogObjectTools, Log, All);
 
diff --git a/Engine/Source/Editor/UnrealEd/Private/PlayLevel.cpp b/Engine/Source/Editor/UnrealEd/Private/PlayLevel.cpp
--- a/Engine/Source/Editor/UnrealEd/Private/PlayLevel.cpp
+++ b/Engine/Source/Editor/UnrealEd/Private/PlayLevel.cpp
@@ -59,6 +59,7 @@
 #include "GameMapsSettings.h"
 #include "GameFramework/Pawn.h"
 #include "GameDelegates.h"
+#include "GeneralProjectSettings.h"
 
 DEFINE_LOG_CATEGORY_STATIC(LogPlayLevel, Log, All);
 
diff --git a/Engine/Source/Editor/UnrealEd/Private/SoundCueGraphSchema.cpp b/Engine/Source/Editor/UnrealEd/Private/SoundCueGraphSchema.cpp
--- a/Engine/Source/Editor/UnrealEd/Private/SoundCueGraphSchema.cpp
+++ b/Engine/Source/Editor/UnrealEd/Private/SoundCueGraphSchema.cpp
@@ -14,6 +14,7 @@
 #include "ScopedTransaction.h"
 #include "GraphEditor.h"
 #include "Sound/SoundCue.h"
+#include "Engine/Selection.h"
 
 #define LOCTEXT_NAMESPACE "SoundCueSchema"
 
diff --git a/Engine/Source/Editor/UnrealEd/Private/StaticLightingSystem/StaticLightingSystem.cpp b/Engine/Source/Editor/UnrealEd/Private/StaticLightingSystem/StaticLightingSystem.cpp
--- a/Engine/Source/Editor/UnrealEd/Private/StaticLightingSystem/StaticLightingSystem.cpp
+++ b/Engine/Source/Editor/UnrealEd/Private/StaticLightingSystem/StaticLightingSystem.cpp
@@ -35,6 +35,9 @@ DEFINE_LOG_CATEGORY(LogStaticLightingSystem);
 #include "SNotificationList.h"
 #include "NotificationManager.h"
 #include "Engine/GeneratedMeshAreaLight.h"
+#include "Engine/LevelStreaming.h"
+#include "Engine/Selection.h"
+#include "Components/SkyLightComponent.h"
 
 #define LOCTEXT_NAMESPACE "StaticLightingSystem"
 
diff --git a/Engine/Source/Editor/UnrealEd/Private/Tests/AutomationEditorCommon.cpp b/Engine/Source/Editor/UnrealEd/Private/Tests/AutomationEditorCommon.cpp
--- a/Engine/Source/Editor/UnrealEd/Private/Tests/AutomationEditorCommon.cpp
+++ b/Engine/Source/Editor/UnrealEd/Private/Tests/AutomationEditorCommon.cpp
@@ -60,6 +60,7 @@
 #include "CookOnTheSide/CookOnTheFlyServer.h"
 #include "Interfaces/ILauncherProfile.h"
 #include "LightingBuildOptions.h"
+#include "Engine/StaticMeshActor.h"
 
 #define COOK_TIMEOUT 3600
 
diff --git a/Engine/Source/Editor/UnrealEd/Private/Tests/BlueprintAutomationTests.cpp b/Engine/Source/Editor/UnrealEd/Private/Tests/BlueprintAutomationTests.cpp
--- a/Engine/Source/Editor/UnrealEd/Private/Tests/BlueprintAutomationTests.cpp
+++ b/Engine/Source/Editor/UnrealEd/Private/Tests/BlueprintAutomationTests.cpp
@@ -13,6 +13,8 @@
 #include "Editor/GraphEditor/Public/GraphDiffControl.h"
 #include "BlueprintSupport.h" // for FScopedClassDependencyGather::GetCachedDependencies()
 #include "EdGraphSchema_K2.h"
+#include "EngineUtils.h"
+#include "GameFramework/SaveGame.h"
 
 DEFINE_LOG_CATEGORY_STATIC(LogBlueprintAutomationTests, Log, All);
 
diff --git a/Engine/Source/Editor/UnrealEd/Private/Tests/EditorBuildPromotionTests.cpp b/Engine/Source/Editor/UnrealEd/Private/Tests/EditorBuildPromotionTests.cpp
--- a/Engine/Source/Editor/UnrealEd/Private/Tests/EditorBuildPromotionTests.cpp
+++ b/Engine/Source/Editor/UnrealEd/Private/Tests/EditorBuildPromotionTests.cpp
@@ -59,6 +59,15 @@
 #include "ObjectTools.h"
 #include "RHI.h"
 
+#include "Engine/SimpleConstructionScript.h"
+#include "Engine/SCS_Node.h"
+#include "Engine/StaticMeshActor.h"
+#include "Engine/PointLight.h"
+#include "GameFramework/PlayerStart.h"
+#include "GameFramework/WorldSettings.h"
+#include "Engine/DirectionalLight.h"
+#include "Distributions/DistributionVectorUniform.h"
+
 #define LOCTEXT_NAMESPACE "EditorBuildPromotionTests"
 
 DEFINE_LOG_CATEGORY_STATIC(LogEditorBuildPromotionTests, Log, All);
diff --git a/Engine/Source/Editor/UnrealEd/Private/UnrealEdEngine.cpp b/Engine/Source/Editor/UnrealEd/Private/UnrealEdEngine.cpp
--- a/Engine/Source/Editor/UnrealEd/Private/UnrealEdEngine.cpp
+++ b/Engine/Source/Editor/UnrealEd/Private/UnrealEdEngine.cpp
@@ -32,6 +32,10 @@
 #include "NotificationManager.h"
 #include "SNotificationList.h"
 #include "UObject/UObjectThreadContext.h"
+#include "Components/BillboardComponent.h"
+#include "Components/ArrowComponent.h"
+#include "Engine/Selection.h"
+#include "EngineUtils.h"
 
 DEFINE_LOG_CATEGORY_STATIC(LogUnrealEdEngine, Log, All);
 
@@ -432,7 +436,7 @@ void UUnrealEdEngine::OnPackageDirtyStateUpdated( UPackage* Pkg)
 	const FString PackageName = Package->GetName();
 
 	// Alert the user if they have modified a package that won't be able to be saved because
-	// it's already been saved with an engine version that is newer than the current one.	
+	// it's already been saved with an engine version that is newer than the current one.
 	if (!FUObjectThreadContext::Get().IsRoutingPostLoad && Package->IsDirty() && !PackagesCheckedForEngineVersion.Contains(PackageName))
 	{
 		EWriteDisallowedWarningState WarningStateToSet = WDWS_WarningUnnecessary;
diff --git a/Engine/Source/Editor/UnrealEd/Private/UnrealEdMisc.cpp b/Engine/Source/Editor/UnrealEd/Private/UnrealEdMisc.cpp
--- a/Engine/Source/Editor/UnrealEd/Private/UnrealEdMisc.cpp
+++ b/Engine/Source/Editor/UnrealEd/Private/UnrealEdMisc.cpp
@@ -38,6 +38,9 @@
 #include "PerformanceMonitor.h"
 #include "Engine/WorldComposition.h"
 #include "FeaturePack.h"
+#include "GameMapsSettings.h"
+#include "GeneralProjectSettings.h"
+#include "Lightmass/LightmappedSurfaceCollection.h"
 
 #define LOCTEXT_NAMESPACE "UnrealEd"
 
diff --git a/Engine/Source/Editor/WorldBrowser/Private/StreamingLevels/StreamingLevelModel.cpp b/Engine/Source/Editor/WorldBrowser/Private/StreamingLevels/StreamingLevelModel.cpp
--- a/Engine/Source/Editor/WorldBrowser/Private/StreamingLevels/StreamingLevelModel.cpp
+++ b/Engine/Source/Editor/WorldBrowser/Private/StreamingLevels/StreamingLevelModel.cpp
@@ -5,6 +5,7 @@
 #include "StreamingLevelCollectionModel.h"
 #include "StreamingLevelModel.h"
 #include "Engine/LevelStreaming.h"
+#include "AssetData.h"
 
 #define LOCTEXT_NAMESPACE "WorldBrowser"
 
diff --git a/Engine/Source/Editor/WorldBrowser/Private/Tiles/WorldTileModel.cpp b/Engine/Source/Editor/WorldBrowser/Private/Tiles/WorldTileModel.cpp
--- a/Engine/Source/Editor/WorldBrowser/Private/Tiles/WorldTileModel.cpp
+++ b/Engine/Source/Editor/WorldBrowser/Private/Tiles/WorldTileModel.cpp
@@ -13,6 +13,7 @@
 #include "Engine/LevelStreaming.h"
 #include "Engine/LevelStreamingKismet.h"
 #include "Landscape.h"
+#include "AssetData.h"
 
 
 #define LOCTEXT_NAMESPACE "WorldBrowser"
diff --git a/Engine/Source/Programs/UnrealHeaderTool/Private/ClassMaps.cpp b/Engine/Source/Programs/UnrealHeaderTool/Private/ClassMaps.cpp
--- a/Engine/Source/Programs/UnrealHeaderTool/Private/ClassMaps.cpp
+++ b/Engine/Source/Programs/UnrealHeaderTool/Private/ClassMaps.cpp
@@ -3,6 +3,8 @@
 #include "UnrealHeaderTool.h"
 
 #include "ClassMaps.h"
+#include "UnrealSourceFile.h"
+#include "UnrealTypeDefinitionInfo.h"
 
 TMap<FString, TSharedRef<FUnrealSourceFile> > GUnrealSourceFilesMap;
 TMap<UField*, TSharedRef<FUnrealTypeDefinitionInfo> > GTypeDefinitionInfoMap;
diff --git a/Engine/Source/Programs/UnrealHeaderTool/Private/ClassMaps.h b/Engine/Source/Programs/UnrealHeaderTool/Private/ClassMaps.h
--- a/Engine/Source/Programs/UnrealHeaderTool/Private/ClassMaps.h
+++ b/Engine/Source/Programs/UnrealHeaderTool/Private/ClassMaps.h
@@ -10,6 +10,8 @@
 #include "UnrealSourceFile.h"
 
 struct FManifestModule;
+class FUnrealSourceFile;
+class FUnrealTypeDefinitionInfo;
 
 /** Structure that holds class meta data generated from its UCLASS declaration */
 struct FClassDeclarationMetaData
diff --git a/Engine/Source/Programs/UnrealHeaderTool/Private/Scope.cpp b/Engine/Source/Programs/UnrealHeaderTool/Private/Scope.cpp
--- a/Engine/Source/Programs/UnrealHeaderTool/Private/Scope.cpp
+++ b/Engine/Source/Programs/UnrealHeaderTool/Private/Scope.cpp
@@ -2,6 +2,7 @@
 
 #include "UnrealHeaderTool.h"
 #include "ParserHelper.h"
+#include "Scope.h"
 
 extern FCompilerMetadataManager GScriptHelper;
 
diff --git a/Engine/Source/Runtime/Engine/Classes/DeviceProfiles/DeviceProfile.h b/Engine/Source/Runtime/Engine/Classes/DeviceProfiles/DeviceProfile.h
--- a/Engine/Source/Runtime/Engine/Classes/DeviceProfiles/DeviceProfile.h
+++ b/Engine/Source/Runtime/Engine/Classes/DeviceProfiles/DeviceProfile.h
@@ -6,6 +6,7 @@
 
 #pragma once
 
+#include "Engine/TextureLODSettings.h"
 #include "DeviceProfile.generated.h"
 
 
diff --git a/Engine/Source/Runtime/Engine/Classes/GameFramework/GameMode.h b/Engine/Source/Runtime/Engine/Classes/GameFramework/GameMode.h
--- a/Engine/Source/Runtime/Engine/Classes/GameFramework/GameMode.h
+++ b/Engine/Source/Runtime/Engine/Classes/GameFramework/GameMode.h
@@ -3,6 +3,7 @@
 #pragma once
 #include "TimerManager.h"
 #include "GameFramework/Info.h"
+#include "PlayerStart.h"
 #include "GameMode.generated.h"
 
 class FDebugDisplayInfo;
diff --git a/Engine/Source/Runtime/Engine/Private/Interpolation.cpp b/Engine/Source/Runtime/Engine/Private/Interpolation.cpp
--- a/Engine/Source/Runtime/Engine/Private/Interpolation.cpp
+++ b/Engine/Source/Runtime/Engine/Private/Interpolation.cpp
@@ -84,6 +84,8 @@
 #include "Engine/Light.h"
 #include "Components/DecalComponent.h"
 
+#include "Sound/SoundCue.h"
+
 DEFINE_LOG_CATEGORY(LogMatinee);
 
 // Priority with which to display sounds triggered by Matinee sound tracks.
diff --git a/Engine/Source/Runtime/Engine/Private/SoundClass.cpp b/Engine/Source/Runtime/Engine/Private/SoundClass.cpp
--- a/Engine/Source/Runtime/Engine/Private/SoundClass.cpp
+++ b/Engine/Source/Runtime/Engine/Private/SoundClass.cpp
@@ -11,6 +11,7 @@
 #include "SNotificationList.h"
 #include "NotificationManager.h"
 #endif
+#include "Sound/SoundMix.h"
 
 /*-----------------------------------------------------------------------------
 	USoundClass implementation.
diff --git a/Engine/Source/Runtime/Engine/Public/Net/NetworkProfiler.h b/Engine/Source/Runtime/Engine/Public/Net/NetworkProfiler.h
--- a/Engine/Source/Runtime/Engine/Public/Net/NetworkProfiler.h
+++ b/Engine/Source/Runtime/Engine/Public/Net/NetworkProfiler.h
@@ -8,6 +8,7 @@
 #define UNREAL_NETWORK_PROFILER_H
 
 class UNetConnection;
+class FInternetAddr;
 
 #if USE_NETWORK_PROFILER 
 
diff --git a/Engine/Source/Runtime/Engine/Public/VisualLogger/VisualLogger.h b/Engine/Source/Runtime/Engine/Public/VisualLogger/VisualLogger.h
--- a/Engine/Source/Runtime/Engine/Public/VisualLogger/VisualLogger.h
+++ b/Engine/Source/Runtime/Engine/Public/VisualLogger/VisualLogger.h
@@ -2,6 +2,7 @@
 #pragma once
 #include "VisualLoggerTypes.h"
 #include "EngineDefines.h"
+#include "EngineStats.h"
 
 #if ENABLE_VISUAL_LOG
 
diff --git a/Engine/Source/Runtime/Foliage/Public/ProceduralFoliageVolume.h b/Engine/Source/Runtime/Foliage/Public/ProceduralFoliageVolume.h
--- a/Engine/Source/Runtime/Foliage/Public/ProceduralFoliageVolume.h
+++ b/Engine/Source/Runtime/Foliage/Public/ProceduralFoliageVolume.h
@@ -2,6 +2,7 @@
 
 #pragma once
 
+#include "GameFramework/Volume.h"
 #include "ProceduralFoliageVolume.generated.h"
 
 class UProceduralFoliageComponent;
diff --git a/Engine/Source/Runtime/GameplayAbilities/Public/GameplayAbilityBlueprintGeneratedClass.h b/Engine/Source/Runtime/GameplayAbilities/Public/GameplayAbilityBlueprintGeneratedClass.h
--- a/Engine/Source/Runtime/GameplayAbilities/Public/GameplayAbilityBlueprintGeneratedClass.h
+++ b/Engine/Source/Runtime/GameplayAbilities/Public/GameplayAbilityBlueprintGeneratedClass.h
@@ -2,6 +2,7 @@
 
 #pragma once
 
+#include "Engine/BlueprintGeneratedClass.h"
 #include "GameplayAbilityBlueprintGeneratedClass.generated.h"
 
 UCLASS()
diff --git a/Engine/Source/Runtime/Renderer/Private/LightSceneInfo.h b/Engine/Source/Runtime/Renderer/Private/LightSceneInfo.h
--- a/Engine/Source/Runtime/Renderer/Private/LightSceneInfo.h
+++ b/Engine/Source/Runtime/Renderer/Private/LightSceneInfo.h
@@ -8,6 +8,7 @@
 #define __LIGHTSCENEINFO_H__
 
 #include "StaticArray.h"
+#include "SceneManagement.h"
 
 /**
  * The information needed to cull a light-primitive interaction.
diff --git a/Engine/Source/Runtime/Renderer/Private/SceneCore.h b/Engine/Source/Runtime/Renderer/Private/SceneCore.h
--- a/Engine/Source/Runtime/Renderer/Private/SceneCore.h
+++ b/Engine/Source/Runtime/Renderer/Private/SceneCore.h
@@ -5,6 +5,7 @@
 =============================================================================*/
 
 #pragma once
+#include "MeshBatch.h"
 
 // Forward declarations.
 class FStaticMesh;
diff --git a/Engine/Source/Runtime/Renderer/Private/ShadowRendering.h b/Engine/Source/Runtime/Renderer/Private/ShadowRendering.h
--- a/Engine/Source/Runtime/Renderer/Private/ShadowRendering.h
+++ b/Engine/Source/Runtime/Renderer/Private/ShadowRendering.h
@@ -8,6 +8,7 @@
 #define __ShadowRendering_H__
 
 #include "ShaderParameterUtils.h"
+#include "SceneCore.h"
 
 // Forward declarations.
 class FProjectedShadowInfo;
diff --git a/Engine/Source/Runtime/Renderer/Private/StaticMeshDrawList.inl b/Engine/Source/Runtime/Renderer/Private/StaticMeshDrawList.inl
--- a/Engine/Source/Runtime/Renderer/Private/StaticMeshDrawList.inl
+++ b/Engine/Source/Runtime/Renderer/Private/StaticMeshDrawList.inl
@@ -8,6 +8,7 @@
 #define __STATICMESHDRAWLIST_INL__
 
 #include "RHICommandList.h"
+#include "EngineStats.h"
 
 // Expensive
 #define PER_MESH_DRAW_STATS 0
diff --git a/Engine/Source/Runtime/Slate/Public/Widgets/Views/SListView.h b/Engine/Source/Runtime/Slate/Public/Widgets/Views/SListView.h
--- a/Engine/Source/Runtime/Slate/Public/Widgets/Views/SListView.h
+++ b/Engine/Source/Runtime/Slate/Public/Widgets/Views/SListView.h
@@ -2,6 +2,7 @@
 
 #pragma once
 #include "STableViewBase.h"
+#include "TableViewTypeTraits.h"
 
 
 /**
diff --git a/Engine/Source/Runtime/UMG/Private/Components/EditableText.cpp b/Engine/Source/Runtime/UMG/Private/Components/EditableText.cpp
--- a/Engine/Source/Runtime/UMG/Private/Components/EditableText.cpp
+++ b/Engine/Source/Runtime/UMG/Private/Components/EditableText.cpp
@@ -1,6 +1,7 @@
 // Copyright 1998-2015 Epic Games, Inc. All Rights Reserved.
 
 #include "UMGPrivatePCH.h"
+#include "SlateFontInfo.h"
 
 #define LOCTEXT_NAMESPACE "UMG"
 
diff --git a/Engine/Source/Runtime/UMG/Private/Components/EditableTextBox.cpp b/Engine/Source/Runtime/UMG/Private/Components/EditableTextBox.cpp
--- a/Engine/Source/Runtime/UMG/Private/Components/EditableTextBox.cpp
+++ b/Engine/Source/Runtime/UMG/Private/Components/EditableTextBox.cpp
@@ -1,6 +1,7 @@
 // Copyright 1998-2015 Epic Games, Inc. All Rights Reserved.
 
 #include "UMGPrivatePCH.h"
+#include "SlateFontInfo.h"
 
 #define LOCTEXT_NAMESPACE "UMG"
 
diff --git a/Engine/Source/Runtime/UMG/Private/Components/MultiLineEditableTextBox.cpp b/Engine/Source/Runtime/UMG/Private/Components/MultiLineEditableTextBox.cpp
--- a/Engine/Source/Runtime/UMG/Private/Components/MultiLineEditableTextBox.cpp
+++ b/Engine/Source/Runtime/UMG/Private/Components/MultiLineEditableTextBox.cpp
@@ -1,6 +1,7 @@
 // Copyright 1998-2015 Epic Games, Inc. All Rights Reserved.
 
 #include "UMGPrivatePCH.h"
+#include "SlateFontInfo.h"
 
 #define LOCTEXT_NAMESPACE "UMG"
 
diff --git a/Engine/Source/Runtime/UMG/Private/WidgetBlueprintLibrary.cpp b/Engine/Source/Runtime/UMG/Private/WidgetBlueprintLibrary.cpp
--- a/Engine/Source/Runtime/UMG/Private/WidgetBlueprintLibrary.cpp
+++ b/Engine/Source/Runtime/UMG/Private/WidgetBlueprintLibrary.cpp
@@ -4,6 +4,8 @@
 #include "Slate/SlateBrushAsset.h"
 #include "Slate/UMGDragDropOp.h"
 #include "WidgetBlueprintLibrary.h"
+#include "EngineGlobals.h"
+#include "Engine/Engine.h"
 
 #define LOCTEXT_NAMESPACE "UMG"
 
diff --git a/Engine/Source/Runtime/UMG/Public/Components/Viewport.h b/Engine/Source/Runtime/UMG/Public/Components/Viewport.h
--- a/Engine/Source/Runtime/UMG/Public/Components/Viewport.h
+++ b/Engine/Source/Runtime/UMG/Public/Components/Viewport.h
@@ -4,6 +4,7 @@
 
 #include "Camera/CameraTypes.h"
 #include "ContentWidget.h"
+#include "SceneManagement.h"
 #include "Viewport.generated.h"
 
 class FPreviewScene;
