Converted the data table editor to use SListView

UETOOL-281 - Improve data table workflows

This is needed as the changes that had been made to SSplitter have been moved to SHeaderRow instead, as they were too specific to go into the generic splitter.

This means that the data table editor now has two list views (one for the row names, one for everything else) which are kept in sync. We have two so that the row names can be anchored to the left of the screen when scrolling the other columns.

diff --git a/Engine/Source/Editor/DataTableEditor/Private/DataTableEditor.cpp b/Engine/Source/Editor/DataTableEditor/Private/DataTableEditor.cpp
--- a/Engine/Source/Editor/DataTableEditor/Private/DataTableEditor.cpp
+++ b/Engine/Source/Editor/DataTableEditor/Private/DataTableEditor.cpp
@@ -14,8 +14,47 @@
  
 #define LOCTEXT_NAMESPACE "DataTableEditor"
 
-const FName FDataTableEditor::DataTableTabId( TEXT( "DataTableEditor_DataTable" ) );
-const FName FDataTableEditor::RowEditorTabId(TEXT("DataTableEditor_RowEditor"));
+const FName FDataTableEditor::DataTableTabId("DataTableEditor_DataTable");
+const FName FDataTableEditor::RowEditorTabId("DataTableEditor_RowEditor");
+const FName FDataTableEditor::RowNameColumnId("RowName");
+
+class SDataTableListViewRow : public SMultiColumnTableRow<FDataTableEditorRowListViewDataPtr>
+{
+public:
+	SLATE_BEGIN_ARGS(SDataTableListViewRow) {}
+		/** The widget that owns the tree.  We'll only keep a weak reference to it. */
+		SLATE_ARGUMENT(TSharedPtr<FDataTableEditor>, DataTableEditor)
+		/** The list item for this row */
+		SLATE_ARGUMENT(FDataTableEditorRowListViewDataPtr, Item)
+	SLATE_END_ARGS()
+
+	/** Construct function for this widget */
+	void Construct(const FArguments& InArgs, const TSharedRef<STableViewBase>& InOwnerTableView)
+	{
+		DataTableEditor = InArgs._DataTableEditor;
+		Item = InArgs._Item;
+		SMultiColumnTableRow<FDataTableEditorRowListViewDataPtr>::Construct(
+			FSuperRowType::FArguments()
+				.Style(FEditorStyle::Get(), "DataTableEditor.CellListViewRow"), 
+			InOwnerTableView
+			);
+	}
+
+	/** Overridden from SMultiColumnTableRow.  Generates a widget for this column of the list view. */
+	virtual TSharedRef<SWidget> GenerateWidgetForColumn(const FName& ColumnName) override
+	{
+		TSharedPtr<FDataTableEditor> DataTableEditorPtr = DataTableEditor.Pin();
+		return (DataTableEditorPtr.IsValid())
+			? DataTableEditorPtr->MakeCellWidget(Item, IndexInList, ColumnName)
+			: SNullWidget::NullWidget;
+	}
+
+private:
+	/** Weak reference to the data table editor that owns our list */
+	TWeakPtr<FDataTableEditor> DataTableEditor;
+	/** The item associated with this row of data */
+	FDataTableEditorRowListViewDataPtr Item;
+};
 
 void FDataTableEditor::RegisterTabSpawners(const TSharedRef<class FTabManager>& TabManager)
 {
@@ -45,7 +84,7 @@ FDataTableEditor::~FDataTableEditor()
 	if (DataTable.IsValid())
 	{
 		SaveLayoutData();
-}
+	}
 }
 
 void FDataTableEditor::PreChange(const class UUserDefinedStruct* Struct, FStructureEditorUtils::EStructureEditorChangeInfo Info)
@@ -56,8 +95,11 @@ void FDataTableEditor::PostChange(const class UUserDefinedStruct* Struct, FStruc
 {
 	if (Struct && DataTable.IsValid() && (DataTable->RowStruct == Struct))
 	{
-		CachedDataTable.Empty();
-		ReloadVisibleData();
+		// We need to cache and restore the selection here as RefreshCachedDataTable will re-create the list view items
+		const FName CachedSelection = HighlightedRowName;
+		HighlightedRowName = NAME_None;
+		RefreshCachedDataTable();
+		RestoreCachedSelection(CachedSelection, true/*bUpdateEvenIfValid*/);
 	}
 }
 
@@ -70,8 +112,11 @@ void FDataTableEditor::PostChange(const UDataTable* Changed, FDataTableEditorUti
 	FStringAssetReference::InvalidateTag(); // Should be removed after UE-5615 is fixed
 	if (Changed == DataTable.Get())
 	{
-		CachedDataTable.Empty();
-		ReloadVisibleData();
+		// We need to cache and restore the selection here as RefreshCachedDataTable will re-create the list view items
+		const FName CachedSelection = HighlightedRowName;
+		HighlightedRowName = NAME_None;
+		RefreshCachedDataTable();
+		RestoreCachedSelection(CachedSelection, true/*bUpdateEvenIfValid*/);
 	}
 }
 
@@ -121,19 +166,17 @@ FText FDataTableEditor::GetBaseToolkitName() const
 	return LOCTEXT( "AppLabel", "DataTable Editor" );
 }
 
-
 FString FDataTableEditor::GetWorldCentricTabPrefix() const
 {
 	return LOCTEXT("WorldCentricTabPrefix", "DataTable ").ToString();
 }
 
-
 FLinearColor FDataTableEditor::GetWorldCentricTabColorScale() const
 {
 	return FLinearColor( 0.0f, 0.0f, 0.2f, 0.5f );
 }
 
-FSlateColor FDataTableEditor::GetRowColor(FName RowName) const
+FSlateColor FDataTableEditor::GetRowTextColor(FName RowName) const
 {
 	if (RowName == HighlightedRowName)
 	{
@@ -142,18 +185,12 @@ FSlateColor FDataTableEditor::GetRowColor(FName RowName) const
 	return FSlateColor::UseForeground();
 }
 
-FReply FDataTableEditor::OnRowClicked(const FGeometry&, const FPointerEvent&, FName RowName)
+FOptionalSize FDataTableEditor::GetRowNameColumnWidth() const
 {
-	if (HighlightedRowName != RowName)
-	{
-		SetHighlightedRow(RowName);
-		CallbackOnRowHighlighted.ExecuteIfBound(HighlightedRowName);
-	}
-
-	return FReply::Handled();
+	return FOptionalSize(RowNameColumnWidth);
 }
 
-float FDataTableEditor::GetColumnWidth(const int32 ColumnIndex)
+float FDataTableEditor::GetColumnWidth(const int32 ColumnIndex) const
 {
 	if (ColumnWidths.IsValidIndex(ColumnIndex))
 	{
@@ -188,7 +225,7 @@ void FDataTableEditor::OnColumnResized(const float NewWidth, const int32 ColumnI
 				LayoutColumnWidths = LayoutData->GetObjectField(TEXT("ColumnWidths"));
 			}
 
-			const FString& ColumnName = CachedRawColumnNames[ColumnIndex];
+			const FString& ColumnName = AvailableColumns[ColumnIndex]->ColumnId.ToString();
 			LayoutColumnWidths->SetNumberField(ColumnName, NewWidth);
 		}
 	}
@@ -230,192 +267,220 @@ void FDataTableEditor::SaveLayoutData()
 	}
 }
 
-TSharedRef<SWidget> FDataTableEditor::CreateGridPanel()
+TSharedRef<ITableRow> FDataTableEditor::MakeRowNameWidget(FDataTableEditorRowListViewDataPtr InRowDataPtr, const TSharedRef<STableViewBase>& OwnerTable)
 {
-	TSharedRef<SVerticalBox> VerticalBox = SNew(SVerticalBox);
+	return
+		SNew(STableRow<FDataTableEditorRowListViewDataPtr>, OwnerTable)
+		.Style(FEditorStyle::Get(), "DataTableEditor.NameListViewRow")
+		[
+			SNew(SBox)
+			.Padding(FMargin(4, 2, 4, 2))
+			[
+				SNew(STextBlock)
+				.ColorAndOpacity(this, &FDataTableEditor::GetRowTextColor, InRowDataPtr->RowId)
+				.Text(InRowDataPtr->DisplayName)
+			]
+		];
+}
 
-	if (DataTable.IsValid())
+TSharedRef<ITableRow> FDataTableEditor::MakeRowWidget(FDataTableEditorRowListViewDataPtr InRowDataPtr, const TSharedRef<STableViewBase>& OwnerTable)
+{
+	return
+		SNew(SDataTableListViewRow, OwnerTable)
+		.DataTableEditor(SharedThis(this))
+		.Item(InRowDataPtr);
+}
+
+TSharedRef<SWidget> FDataTableEditor::MakeCellWidget(FDataTableEditorRowListViewDataPtr InRowDataPtr, const int32 InRowIndex, const FName& InColumnId)
+{
+	int32 ColumnIndex = 0;
+	for (; ColumnIndex < AvailableColumns.Num(); ++ColumnIndex)
 	{
-		if (CachedDataTable.Num() == 0)
+		const FDataTableEditorColumnHeaderDataPtr& ColumnData = AvailableColumns[ColumnIndex];
+		if (ColumnData->ColumnId == InColumnId)
 		{
-			CachedDataTable = DataTable->GetTableData();
-			check(CachedDataTable.Num() > 0); // There should always be at least the column titles row
+			break;
+		}
+	}
 
-			CachedRawColumnNames = DataTable->GetUniqueColumnTitles();
+	// Valid column ID?
+	if (AvailableColumns.IsValidIndex(ColumnIndex))
+	{
+		return SNew(SBox)
+			.Padding(FMargin(4, 2, 4, 2))
+			[
+				SNew(STextBlock)
+				.TextStyle(FEditorStyle::Get(), "DataTableEditor.CellText")
+				.ColorAndOpacity(this, &FDataTableEditor::GetRowTextColor, InRowDataPtr->RowId)
+				.Text(InRowDataPtr->CellData[ColumnIndex])
+				.ToolTipText(FText::Format(LOCTEXT("ColumnRowNameFmt", "{0}: {1}"), AvailableColumns[ColumnIndex]->DisplayName, InRowDataPtr->CellData[ColumnIndex]))
+			];
+	}
 
-			RowsVisibility.SetNum(CachedDataTable.Num());
-			for (bool& RowVisibility : RowsVisibility)
-			{
-				RowVisibility = true;
-			}
+	return SNullWidget::NullWidget;
+}
 
-			ColumnWidths.SetNum(CachedRawColumnNames.Num());
-			
-			// Load the persistent column widths from the layout data
-			{
-				const TSharedPtr<FJsonObject>* LayoutColumnWidths = nullptr;
-				if (LayoutData.IsValid() && LayoutData->TryGetObjectField(TEXT("ColumnWidths"), LayoutColumnWidths))
-				{
-					for(int32 ColumnIndex = 0; ColumnIndex < CachedRawColumnNames.Num(); ++ColumnIndex)
-					{
-						const FString& ColumnName = CachedRawColumnNames[ColumnIndex];
-
-						double LayoutColumnWidth = 0.0f;
-						if ((*LayoutColumnWidths)->TryGetNumberField(ColumnName, LayoutColumnWidth))
-						{
-							FColumnWidth& ColumnWidth = ColumnWidths[ColumnIndex];
-							ColumnWidth.bIsAutoSized = false;
-							ColumnWidth.CurrentWidth = static_cast<float>(LayoutColumnWidth);
-						}
-					}
-		}
-			}
-		}
+void FDataTableEditor::OnRowNamesListViewScrolled(double InScrollOffset)
+{
+	// Synchronize the list views
+	CellsListView->SetScrollOffset(InScrollOffset);
+}
 
-		check(CachedDataTable.Num() > 0 && CachedDataTable.Num() == RowsVisibility.Num());
-		check(CachedDataTable.Num() > 0 && CachedDataTable[0].Num() == ColumnWidths.Num());
-		check(CachedRawColumnNames.Num() > 0 && CachedRawColumnNames.Num() == ColumnWidths.Num());
+void FDataTableEditor::OnCellsListViewScrolled(double InScrollOffset)
+{
+	// Synchronize the list views
+	RowNamesListView->SetScrollOffset(InScrollOffset);
+}
 
-		int32 VisibleRowIndex = 0;
-		TArray<FString>& ColumnTitles = CachedDataTable[0];
-		for(int32 RowIndex = 0; RowIndex < CachedDataTable.Num(); ++RowIndex)
-		{
-			if (RowsVisibility[RowIndex])
-			{
-				const bool bIsHeader = (RowIndex == 0);
-				const FLinearColor RowColor = (VisibleRowIndex % 2 == 0) ? FLinearColor::Gray : FLinearColor::Black;
-
-				TArray<FString>& Row = CachedDataTable[RowIndex];
-				FName RowName(*Row[0]);
-				TAttribute<FSlateColor> ForegroundColor = bIsHeader
-					? FSlateColor::UseForeground()
-					: TAttribute<FSlateColor>::Create(
-					TAttribute<FSlateColor>::FGetter::CreateSP(this, &FDataTableEditor::GetRowColor, RowName));
-
-				auto RowClickCallback = bIsHeader
-					? FPointerEventHandler()
-					: FPointerEventHandler::CreateSP(this, &FDataTableEditor::OnRowClicked, RowName);
-
-				const float SplitterHandleSize = 2.0f;
-				TSharedRef<SSplitter> RowSplitter = SNew(SSplitter)
-					.PhysicalSplitterHandleSize(SplitterHandleSize)
-					.HitDetectionSplitterHandleSize(4.0f);
-
-				VerticalBox->AddSlot()
-				[
-					RowSplitter
-				];
+void FDataTableEditor::OnRowSelectionChanged(FDataTableEditorRowListViewDataPtr InNewSelection, ESelectInfo::Type InSelectInfo)
+{
+	const bool bSelectionChanged = !InNewSelection.IsValid() || InNewSelection->RowId != HighlightedRowName;
+	const FName NewRowName = (InNewSelection.IsValid()) ? InNewSelection->RowId : NAME_None;
 
-				float ColumnDesiredWidth = 0.0f;
-				for(int32 ColumnIndex = 0; ColumnIndex < Row.Num(); ++ColumnIndex)
-				{
-					TSharedPtr<SWidget> ColumnEntryWidget;
-
-					RowSplitter->AddSlot()
-					.Value(TAttribute<float>::Create(TAttribute<float>::FGetter::CreateSP(this, &FDataTableEditor::GetColumnWidth, ColumnIndex)))
-					.OnSlotResized(SSplitter::FOnSlotResized::CreateSP(this, &FDataTableEditor::OnColumnResized, ColumnIndex))
-					.SizeRule(SSplitter::ManualSize)
-					[
-						SAssignNew(ColumnEntryWidget, SBorder)
-						.Padding(2.0f)
-						.BorderImage(FEditorStyle::GetBrush("ToolPanel.GroupBorder"))
-						.BorderBackgroundColor(RowColor)
-						.ForegroundColor(ForegroundColor)
-						.OnMouseButtonDown(RowClickCallback)
-						[
-							SNew(STextBlock)
-							.Text(FText::FromString(Row[ColumnIndex]))
-							.ToolTipText(bIsHeader 
-							?	(FText::Format(LOCTEXT("ColumnHeaderNameFmt", "Column '{0}'"), FText::FromString(ColumnTitles[ColumnIndex]))) 
-							:	(FText::Format(LOCTEXT("ColumnRowNameFmt", "{0}: {1}"), FText::FromString(ColumnTitles[ColumnIndex]), FText::FromString(Row[ColumnIndex]))))
-						]
-					];
+	SetHighlightedRow(NewRowName);
 
-					FColumnWidth& ColumnWidth = ColumnWidths[ColumnIndex];
-					if (ColumnWidth.bIsAutoSized)
-					{
-						ColumnEntryWidget->SlatePrepass(1.0f);
-						ColumnWidth.CurrentWidth = FMath::Max(ColumnWidth.CurrentWidth, ColumnEntryWidget->GetDesiredSize().X + SplitterHandleSize);
-					}
-				}
+	if (bSelectionChanged)
+	{
+		CallbackOnRowHighlighted.ExecuteIfBound(HighlightedRowName);
+	}
+}
 
-				// Dummy splitter slot to allow the last column to be resized
-				RowSplitter->AddSlot()
-				.Value(TAttribute<float>::Create(TAttribute<float>::FGetter::CreateSP(this, &FDataTableEditor::GetColumnWidth, static_cast<int32>(INDEX_NONE))))
-				.OnSlotResized(SSplitter::FOnSlotResized::CreateSP(this, &FDataTableEditor::OnColumnResized, static_cast<int32>(INDEX_NONE)))
-				.SizeRule(SSplitter::ManualSize)
-				[
-					SNullWidget::NullWidget
-				];
+void FDataTableEditor::OnSearchTextChanged(const FText& SearchText)
+{
+	UpdateVisibleRows();
+}
 
-				++VisibleRowIndex;
-			}
+void FDataTableEditor::RefreshCachedDataTable()
+{
+	FDataTableEditorUtils::CacheDataTableForEditing(DataTable.Get(), AvailableColumns, AvailableRows);
+
+	// Update the desired width of the row names column
+	// This prevents it growing or shrinking as you scroll the list view
+	{
+		TSharedRef<FSlateFontMeasure> FontMeasure = FSlateApplication::Get().GetRenderer()->GetFontMeasureService();
+		const FTextBlockStyle& CellTextStyle = FEditorStyle::GetWidgetStyle<FTextBlockStyle>("DataTableEditor.CellText");
+		static const float CellPadding = 10.0f;
+		
+		RowNameColumnWidth = 10.0f;
+		for (const FDataTableEditorRowListViewDataPtr& RowData : AvailableRows)
+		{
+			const float RowNameWidth = FontMeasure->Measure(RowData->DisplayName, CellTextStyle.Font).X + CellPadding;
+			RowNameColumnWidth = FMath::Max(RowNameColumnWidth, RowNameWidth);
 		}
+	}
+
+	// Setup the default auto-sized columns
+	ColumnWidths.SetNum(AvailableColumns.Num());
+	for (int32 ColumnIndex = 0; ColumnIndex < AvailableColumns.Num(); ++ColumnIndex)
+	{
+		const FDataTableEditorColumnHeaderDataPtr& ColumnData = AvailableColumns[ColumnIndex];
+		FColumnWidth& ColumnWidth = ColumnWidths[ColumnIndex];
+		ColumnWidth.CurrentWidth = FMath::Clamp(ColumnData->DesiredColumnWidth, 10.0f, 400.0f); // Clamp auto-sized columns to a reasonable limit
+	}
 
-		// Update the currently selected row to try and ensure it's valid
-		if (HighlightedRowName.IsNone() || !DataTable->RowMap.Contains(HighlightedRowName))
+	// Load the persistent column widths from the layout data
+	{
+		const TSharedPtr<FJsonObject>* LayoutColumnWidths = nullptr;
+		if (LayoutData.IsValid() && LayoutData->TryGetObjectField(TEXT("ColumnWidths"), LayoutColumnWidths))
 		{
-			if (CachedDataTable.Num() > 2)
-			{
-				// 1 is the first row in the table
-				// 0 is the name column
-				SetHighlightedRow(*CachedDataTable[1][0]);
-			}
-			else
+			for(int32 ColumnIndex = 0; ColumnIndex < AvailableColumns.Num(); ++ColumnIndex)
 			{
-				SetHighlightedRow(NAME_None);
+				const FDataTableEditorColumnHeaderDataPtr& ColumnData = AvailableColumns[ColumnIndex];
+
+				double LayoutColumnWidth = 0.0f;
+				if ((*LayoutColumnWidths)->TryGetNumberField(ColumnData->ColumnId.ToString(), LayoutColumnWidth))
+				{
+					FColumnWidth& ColumnWidth = ColumnWidths[ColumnIndex];
+					ColumnWidth.bIsAutoSized = false;
+					ColumnWidth.CurrentWidth = static_cast<float>(LayoutColumnWidth);
+				}
 			}
 		}
 	}
 
-	return VerticalBox;
+	ColumnNamesHeaderRow->ClearColumns();
+	for (int32 ColumnIndex = 0; ColumnIndex < AvailableColumns.Num(); ++ColumnIndex)
+	{
+		const FDataTableEditorColumnHeaderDataPtr& ColumnData = AvailableColumns[ColumnIndex];
+
+		ColumnNamesHeaderRow->AddColumn(
+			SHeaderRow::Column(ColumnData->ColumnId)
+			.DefaultLabel(ColumnData->DisplayName)
+			.ManualWidth(TAttribute<float>::Create(TAttribute<float>::FGetter::CreateSP(this, &FDataTableEditor::GetColumnWidth, ColumnIndex)))
+			.OnWidthChanged(this, &FDataTableEditor::OnColumnResized, ColumnIndex)
+			);
+	}
+
+	UpdateVisibleRows();
 }
 
-void FDataTableEditor::OnSearchTextChanged(const FText& SearchText)
+void FDataTableEditor::UpdateVisibleRows()
 {
-	FString SearchFor = SearchText.ToString();
-	if (!SearchFor.IsEmpty())
+	const FText& ActiveSearchText = (SearchBox.IsValid()) ? SearchBox->GetText() : FText::GetEmpty();
+	if (ActiveSearchText.IsEmptyOrWhitespace())
+	{
+		VisibleRows = AvailableRows;
+	}
+	else
 	{
-		check(CachedDataTable.Num() == RowsVisibility.Num());
+		VisibleRows.Empty(AvailableRows.Num());
 
-		// starting from index 1, because 0 is header
-		for (int32 RowIdx = 1; RowIdx < CachedDataTable.Num(); ++RowIdx)
+		const FString& ActiveSearchString = ActiveSearchText.ToString();
+		for (const FDataTableEditorRowListViewDataPtr& RowData : AvailableRows)
 		{
-			RowsVisibility[RowIdx] = false;
+			bool bPassesFilter = false;
 
-			for (int32 i = 0; i < CachedDataTable[RowIdx].Num(); ++i)
+			if (RowData->DisplayName.ToString().Contains(ActiveSearchString))
 			{
-				if (SearchFor.Len() <= CachedDataTable[RowIdx][i].Len())
+				bPassesFilter = true;
+			}
+			else
+			{
+				for (const FText& CellText : RowData->CellData)
 				{
-					if (CachedDataTable[RowIdx][i].Contains(SearchFor))
+					if (CellText.ToString().Contains(ActiveSearchString))
 					{
-						RowsVisibility[RowIdx] = true;
+						bPassesFilter = true;
 						break;
 					}
 				}
 			}
+
+			if (bPassesFilter)
+			{
+				VisibleRows.Add(RowData);
+			}
 		}
 	}
-	else
-	{
-		for (int32 RowIdx = 0; RowIdx < RowsVisibility.Num(); ++RowIdx)
-		{
-			RowsVisibility[RowIdx] = true;
-		}
-	}
 
-	ReloadVisibleData();
+	RowNamesListView->RequestListRefresh();
+	CellsListView->RequestListRefresh();
+
+	RestoreCachedSelection(HighlightedRowName);
 }
 
-void FDataTableEditor::ReloadVisibleData()
+void FDataTableEditor::RestoreCachedSelection(const FName InCachedSelection, const bool bUpdateEvenIfValid)
 {
-	if (ScrollBoxWidget.IsValid())
+	// Validate the requested selection to see if it matches a known row
+	bool bSelectedRowIsValid = false;
+	if (!InCachedSelection.IsNone())
 	{
-		ScrollBoxWidget->ClearChildren();
-		ScrollBoxWidget->AddSlot()
-			[
-				CreateGridPanel()
-			];
+		bSelectedRowIsValid = VisibleRows.ContainsByPredicate([&InCachedSelection](const FDataTableEditorRowListViewDataPtr& RowData) -> bool
+		{
+			return RowData->RowId == InCachedSelection;
+		});
+	}
+
+	// Apply the new selection (if required)
+	if (!bSelectedRowIsValid)
+	{
+		SetHighlightedRow((VisibleRows.Num() > 1) ? VisibleRows[0]->RowId : NAME_None);
+		CallbackOnRowHighlighted.ExecuteIfBound(HighlightedRowName);
+	}
+	else if (bUpdateEvenIfValid)
+	{
+		SetHighlightedRow(InCachedSelection);
+		CallbackOnRowHighlighted.ExecuteIfBound(HighlightedRowName);
 	}
 }
 
@@ -423,11 +488,43 @@ TSharedRef<SVerticalBox> FDataTableEditor::CreateContentBox()
 {
 	TSharedRef<SScrollBar> HorizontalScrollBar = SNew(SScrollBar)
 		.Orientation(Orient_Horizontal)
-		.Thickness(FVector2D(5, 5));
+		.Thickness(FVector2D(8.0f, 8.0f));
 
 	TSharedRef<SScrollBar> VerticalScrollBar = SNew(SScrollBar)
 		.Orientation(Orient_Vertical)
-		.Thickness(FVector2D(5, 5));
+		.Thickness(FVector2D(8.0f, 8.0f));
+
+	TSharedRef<SHeaderRow> RowNamesHeaderRow = SNew(SHeaderRow);
+	RowNamesHeaderRow->AddColumn(
+		SHeaderRow::Column(RowNameColumnId)
+		.DefaultLabel(FText::GetEmpty())
+		);
+
+	ColumnNamesHeaderRow = SNew(SHeaderRow);
+
+	RowNamesListView = SNew(SListView<FDataTableEditorRowListViewDataPtr>)
+		.ListItemsSource(&VisibleRows)
+		.HeaderRow(RowNamesHeaderRow)
+		.OnGenerateRow(this, &FDataTableEditor::MakeRowNameWidget)
+		.OnListViewScrolled(this, &FDataTableEditor::OnRowNamesListViewScrolled)
+		.OnSelectionChanged(this, &FDataTableEditor::OnRowSelectionChanged)
+		.ScrollbarVisibility(EVisibility::Collapsed)
+		.ConsumeMouseWheel(EConsumeMouseWheel::Always)
+		.SelectionMode(ESelectionMode::Single)
+		.AllowOverscroll(EAllowOverscroll::No);
+
+	CellsListView = SNew(SListView<FDataTableEditorRowListViewDataPtr>)
+		.ListItemsSource(&VisibleRows)
+		.HeaderRow(ColumnNamesHeaderRow)
+		.OnGenerateRow(this, &FDataTableEditor::MakeRowWidget)
+		.OnListViewScrolled(this, &FDataTableEditor::OnCellsListViewScrolled)
+		.OnSelectionChanged(this, &FDataTableEditor::OnRowSelectionChanged)
+		.ExternalScrollbar(VerticalScrollBar)
+		.ConsumeMouseWheel(EConsumeMouseWheel::Always)
+		.SelectionMode(ESelectionMode::Single)
+		.AllowOverscroll(EAllowOverscroll::No);
+
+	RefreshCachedDataTable();
 
 	return SNew(SVerticalBox)
 		+SVerticalBox::Slot()
@@ -440,20 +537,22 @@ TSharedRef<SVerticalBox> FDataTableEditor::CreateContentBox()
 		[
 			SNew(SHorizontalBox)
 			+SHorizontalBox::Slot()
+			.AutoWidth()
+			[
+				SNew(SBox)
+				.WidthOverride(this, &FDataTableEditor::GetRowNameColumnWidth)
+				[
+					RowNamesListView.ToSharedRef()
+				]
+			]
+			+SHorizontalBox::Slot()
 			[
 				SNew(SScrollBox)
 				.Orientation(Orient_Horizontal)
 				.ExternalScrollbar(HorizontalScrollBar)
 				+SScrollBox::Slot()
 				[
-			SAssignNew(ScrollBoxWidget, SScrollBox)
-					.Orientation(Orient_Vertical)
-					.ExternalScrollbar(VerticalScrollBar)
-					.ConsumeMouseWheel(EConsumeMouseWheel::Always) // Always consume the mouse wheel events to prevent the outer scroll box from scrolling
-			+SScrollBox::Slot()
-			[
-						CreateGridPanel()
-					]
+					CellsListView.ToSharedRef()
 				]
 			]
 			+SHorizontalBox::Slot()
@@ -465,7 +564,20 @@ TSharedRef<SVerticalBox> FDataTableEditor::CreateContentBox()
 		+SVerticalBox::Slot()
 		.AutoHeight()
 		[
-			HorizontalScrollBar
+			SNew(SHorizontalBox)
+			+SHorizontalBox::Slot()
+			.AutoWidth()
+			[
+				SNew(SBox)
+				.WidthOverride(this, &FDataTableEditor::GetRowNameColumnWidth)
+				[
+					SNullWidget::NullWidget
+				]
+			]
+			+SHorizontalBox::Slot()
+			[
+				HorizontalScrollBar
+			]
 		];
 }
 
@@ -506,28 +618,58 @@ TSharedRef<SDockTab> FDataTableEditor::SpawnTab_DataTable( const FSpawnTabArgs&
 
 	LoadLayoutData();
 
-	TSharedRef<SVerticalBox> ContentBox = CreateContentBox();
-
-	GridPanelOwner = 
-		SNew(SBorder)
-		.Padding(2)
-		.BorderImage( FEditorStyle::GetBrush( "ToolPanel.GroupBorder" ) )
-		[
-			ContentBox
-		];
-
 	return SNew(SDockTab)
 		.Icon( FEditorStyle::GetBrush("DataTableEditor.Tabs.Properties") )
 		.Label( LOCTEXT("DataTableTitle", "Data Table") )
 		.TabColorScale( GetTabColorScale() )
 		[
-			GridPanelOwner.ToSharedRef()
+			SNew(SBorder)
+			.Padding(2)
+			.BorderImage( FEditorStyle::GetBrush( "ToolPanel.GroupBorder" ) )
+			[
+				CreateContentBox()
+			]
 		];
 }
 
 void FDataTableEditor::SetHighlightedRow(FName Name)
 {
-	HighlightedRowName = Name;
+	if (Name == HighlightedRowName)
+	{
+		return;
+	}
+
+	if (Name.IsNone())
+	{
+		HighlightedRowName = NAME_None;
+
+		// Synchronize the list views
+		RowNamesListView->ClearSelection();
+		CellsListView->ClearSelection();
+	}
+	else
+	{
+		HighlightedRowName = Name;
+
+		FDataTableEditorRowListViewDataPtr* NewSelectionPtr = VisibleRows.FindByPredicate([&Name](const FDataTableEditorRowListViewDataPtr& RowData) -> bool
+		{
+			return RowData->RowId == Name;
+		});
+
+		// Synchronize the list views
+		if (NewSelectionPtr)
+		{
+			RowNamesListView->SetSelection(*NewSelectionPtr);
+			CellsListView->SetSelection(*NewSelectionPtr);
+
+			CellsListView->RequestScrollIntoView(*NewSelectionPtr);
+		}
+		else
+		{
+			RowNamesListView->ClearSelection();
+			CellsListView->ClearSelection();
+		}
+	}
 }
 
 #undef LOCTEXT_NAMESPACE
diff --git a/Engine/Source/Editor/DataTableEditor/Private/DataTableEditor.h b/Engine/Source/Editor/DataTableEditor/Private/DataTableEditor.h
--- a/Engine/Source/Editor/DataTableEditor/Private/DataTableEditor.h
+++ b/Engine/Source/Editor/DataTableEditor/Private/DataTableEditor.h
@@ -18,6 +18,7 @@ class FDataTableEditor : public IDataTableEditor
 	, public FStructureEditorUtils::INotifyOnStructChanged
 	, public FDataTableEditorUtils::INotifyOnDataTableChanged
 {
+	friend class SDataTableListViewRow;
 
 public:
 
@@ -39,9 +40,6 @@ public:
 	/** Destructor */
 	virtual ~FDataTableEditor();
 
-	/* Create the grid panel */
-	TSharedRef<SWidget> CreateGridPanel();
-
 	/** IToolkit interface */
 	virtual FName GetToolkitFName() const override;
 	virtual FText GetBaseToolkitName() const override;
@@ -60,10 +58,16 @@ public:
 
 private:
 
-	void ReloadVisibleData();
+	void RefreshCachedDataTable();
+
+	void UpdateVisibleRows();
+
+	void RestoreCachedSelection(const FName InCachedSelection, const bool bUpdateEvenIfValid = false);
 
 	void OnSearchTextChanged(const FText& SearchText);
 
+	FSlateColor GetRowTextColor(FName RowName) const;
+
 	TSharedRef<SVerticalBox> CreateContentBox();
 
 	TSharedRef<SWidget> CreateRowEditorBox();
@@ -74,11 +78,9 @@ private:
 	/**	Spawns the tab with the Row Editor inside */
 	TSharedRef<SDockTab> SpawnTab_RowEditor(const FSpawnTabArgs& Args);
 
-	FSlateColor GetRowColor(FName RowName) const;
-
-	FReply OnRowClicked(const FGeometry&, const FPointerEvent&, FName RowName);
+	FOptionalSize GetRowNameColumnWidth() const;
 
-	float GetColumnWidth(const int32 ColumnIndex);
+	float GetColumnWidth(const int32 ColumnIndex) const;
 
 	void OnColumnResized(const float NewWidth, const int32 ColumnIndex);
 
@@ -86,6 +88,21 @@ private:
 
 	void SaveLayoutData();
 
+	/** Make the widget for a row name entry in the data table row list view */
+	TSharedRef<ITableRow> MakeRowNameWidget(FDataTableEditorRowListViewDataPtr InRowDataPtr, const TSharedRef<STableViewBase>& OwnerTable);
+
+	/** Make the widget for a row entry in the data table row list view */
+	TSharedRef<ITableRow> MakeRowWidget(FDataTableEditorRowListViewDataPtr InRowDataPtr, const TSharedRef<STableViewBase>& OwnerTable);
+
+	/** Make the widget for a cell entry in the data table row list view */
+	TSharedRef<SWidget> MakeCellWidget(FDataTableEditorRowListViewDataPtr InRowDataPtr, const int32 InRowIndex, const FName& InColumnId);
+
+	void OnRowNamesListViewScrolled(double InScrollOffset);
+
+	void OnCellsListViewScrolled(double InScrollOffset);
+
+	void OnRowSelectionChanged(FDataTableEditorRowListViewDataPtr InNewSelection, ESelectInfo::Type InSelectInfo);
+
 private:
 
 	/** Struct holding information about the current column widths */
@@ -103,26 +120,32 @@ private:
 		float CurrentWidth;
 	};
 
-	/** Cached table data */
-	TArray<TArray<FString> > CachedDataTable;
+	/** Array of the columns that are available for editing */
+	TArray<FDataTableEditorColumnHeaderDataPtr> AvailableColumns;
 
-	/** Cached raw column names */
-	TArray<FString> CachedRawColumnNames;
+	/** Array of the rows that are available for editing */
+	TArray<FDataTableEditorRowListViewDataPtr> AvailableRows;
+	
+	/** Array of the rows that match the active filter(s) */
+	TArray<FDataTableEditorRowListViewDataPtr> VisibleRows;
 
-	/** Visibility of data table rows */
-	TArray<bool> RowsVisibility;
+	/** Header row containing entries for each column in AvailableColumns */
+	TSharedPtr<SHeaderRow> ColumnNamesHeaderRow;
 
-	/** Widths of data table columns */
-	TArray<FColumnWidth> ColumnWidths;
+	/** List view responsible for showing the row names column */
+	TSharedPtr<SListView<FDataTableEditorRowListViewDataPtr>> RowNamesListView;
 
-	/** Search box */
-	TSharedPtr<SWidget> SearchBox;
+	/** List view responsible for showing the rows in VisibleRows for each entry in AvailableColumns */
+	TSharedPtr<SListView<FDataTableEditorRowListViewDataPtr>> CellsListView;
 
-	/** Scroll containing data table */
-	TSharedPtr<SScrollBox> ScrollBoxWidget;
+	/** Width of the row name column */
+	float RowNameColumnWidth;
 
-	/**Border surrounding the GridPanel */
-	TSharedPtr<SBorder>	 GridPanelOwner;
+	/** Widths of data table cell columns */
+	TArray<FColumnWidth> ColumnWidths;
+
+	/** Search box */
+	TSharedPtr<SSearchBox> SearchBox;
 
 	/* The DataTable that is active in the editor */
 	TAssetPtr<UDataTable> DataTable;
@@ -139,4 +162,7 @@ private:
 
 	/**	The tab id for the row editor tab */
 	static const FName RowEditorTabId;
+
+	/** The column id for the row name list view column */
+	static const FName RowNameColumnId;
 };
diff --git a/Engine/Source/Editor/DataTableEditor/Private/SRowEditor.cpp b/Engine/Source/Editor/DataTableEditor/Private/SRowEditor.cpp
--- a/Engine/Source/Editor/DataTableEditor/Private/SRowEditor.cpp
+++ b/Engine/Source/Editor/DataTableEditor/Private/SRowEditor.cpp
@@ -147,28 +147,38 @@ void SRowEditor::RefreshNameList()
 
 void SRowEditor::Restore()
 {
-	if (SelectedName.IsValid())
+	if (!SelectedName.IsValid() || !SelectedName->IsNone())
 	{
-		auto CurrentName = *SelectedName;
-		SelectedName = NULL;
-		for (auto Element : CachedRowNames)
+		if (SelectedName.IsValid())
 		{
-			if (*Element == CurrentName)
+			auto CurrentName = *SelectedName;
+			SelectedName = NULL;
+			for (auto Element : CachedRowNames)
 			{
-				SelectedName = Element;
-				break;
+				if (*Element == CurrentName)
+				{
+					SelectedName = Element;
+					break;
+				}
 			}
 		}
-	}
 
-	if (!SelectedName.IsValid() && CachedRowNames.Num() && CachedRowNames[0].IsValid())
-	{
-		SelectedName = CachedRowNames[0];
-	}
+		if (!SelectedName.IsValid() && CachedRowNames.Num() && CachedRowNames[0].IsValid())
+		{
+			SelectedName = CachedRowNames[0];
+		}
 
-	if (RowComboBox.IsValid())
+		if (RowComboBox.IsValid())
+		{
+			RowComboBox->SetSelectedItem(SelectedName);
+		}
+	}
+	else
 	{
-		RowComboBox->SetSelectedItem(SelectedName);
+		if (RowComboBox.IsValid())
+		{
+			RowComboBox->ClearSelection();
+		}
 	}
 
 	auto FinalName = SelectedName.IsValid() ? *SelectedName : NAME_None;
@@ -211,7 +221,7 @@ TSharedRef<SWidget> SRowEditor::OnGenerateWidget(TSharedPtr<FName> InItem)
 
 void SRowEditor::OnSelectionChanged(TSharedPtr<FName> InItem, ESelectInfo::Type InSeletionInfo)
 {
-	if (InItem != SelectedName)
+	if (InItem.IsValid() && InItem != SelectedName)
 	{
 		CleanBeforeChange();
 
@@ -237,7 +247,7 @@ void SRowEditor::SelectRow(FName InName)
 			NewSelectedName = Name;
 		}
 	}
-	if (!NewSelectedName->IsValid())
+	if (!NewSelectedName.IsValid())
 	{
 		NewSelectedName = MakeShareable(new FName(InName));
 	}
@@ -254,8 +264,8 @@ FReply SRowEditor::OnAddClicked()
 		{
 			NewName.SetNumber(NewName.GetNumber() + 1);
 		}
-		SelectedName = MakeShareable(new FName(NewName));
 		FDataTableEditorUtils::AddRow(DataTable.Get(), NewName);
+		SelectRow(NewName);
 	}
 	return FReply::Handled();
 }
@@ -306,8 +316,8 @@ void SRowEditor::OnRowRenamed(const FText& Text, ETextCommit::Type CommitType)
 		}
 
 		const FName OldName = GetCurrentName();
-		SelectedName = MakeShareable(new FName(NewName));
 		FDataTableEditorUtils::RenameRow(DataTable.Get(), OldName, NewName);
+		SelectRow(NewName);
 	}
 }
 
diff --git a/Engine/Source/Editor/EditorStyle/Private/SlateEditorStyle.cpp b/Engine/Source/Editor/EditorStyle/Private/SlateEditorStyle.cpp
--- a/Engine/Source/Editor/EditorStyle/Private/SlateEditorStyle.cpp
+++ b/Engine/Source/Editor/EditorStyle/Private/SlateEditorStyle.cpp
@@ -5640,6 +5640,37 @@ void FSlateEditorStyle::FStyle::SetupToolkitStyles()
 	// DataTable Editor
 	{
 		Set( "DataTableEditor.Tabs.Properties", new IMAGE_BRUSH( "/Icons/icon_tab_SelectionDetails_16x", Icon16x16 ) );
+
+		Set( "DataTableEditor.CellText", FTextBlockStyle(NormalText)
+			.SetFont( TTF_CORE_FONT("Fonts/Roboto-Regular", 9 ))
+			);
+
+		Set( "DataTableEditor.NameListViewRow", FTableRowStyle(NormalTableRowStyle)
+			.SetEvenRowBackgroundBrush( BOX_BRUSH( "Common/TableViewMajorColumn", 4.f/32.f ) )
+			.SetEvenRowBackgroundHoveredBrush( BOX_BRUSH( "Common/TableViewMajorColumn", 4.f/32.f ) )
+			.SetOddRowBackgroundBrush( BOX_BRUSH( "Common/TableViewMajorColumn", 4.f/32.f ) )
+			.SetOddRowBackgroundHoveredBrush( BOX_BRUSH( "Common/TableViewMajorColumn", 4.f/32.f ) )
+			.SetSelectorFocusedBrush( FSlateNoResource() )
+			.SetActiveBrush( BOX_BRUSH( "Common/TableViewMajorColumn", 4.f/32.f ) )
+			.SetActiveHoveredBrush( BOX_BRUSH( "Common/TableViewMajorColumn", 4.f/32.f ) )
+			.SetInactiveBrush( BOX_BRUSH( "Common/TableViewMajorColumn", 4.f/32.f ) )
+			.SetInactiveHoveredBrush( BOX_BRUSH( "Common/TableViewMajorColumn", 4.f/32.f ) )
+			.SetTextColor( DefaultForeground )
+			.SetSelectedTextColor( DefaultForeground )
+		);
+
+		Set( "DataTableEditor.CellListViewRow", FTableRowStyle(NormalTableRowStyle)
+			.SetEvenRowBackgroundBrush(IMAGE_BRUSH("PropertyView/DetailCategoryMiddle", FVector2D(16, 16), FLinearColor(0.5f, 0.5f, 0.5f)))
+			.SetEvenRowBackgroundHoveredBrush(IMAGE_BRUSH("PropertyView/DetailCategoryMiddle_Hovered", FVector2D(16, 16), FLinearColor(0.5f, 0.5f, 0.5f)))
+			.SetOddRowBackgroundBrush(IMAGE_BRUSH("PropertyView/DetailCategoryMiddle", FVector2D(16, 16), FLinearColor(0.2f, 0.2f, 0.2f)))
+			.SetOddRowBackgroundHoveredBrush(IMAGE_BRUSH("PropertyView/DetailCategoryMiddle_Hovered", FVector2D(16, 16), FLinearColor(0.2f, 0.2f, 0.2f)))
+			.SetActiveBrush( FSlateNoResource() )
+			.SetActiveHoveredBrush( FSlateNoResource() )
+			.SetInactiveBrush( FSlateNoResource() )
+			.SetInactiveHoveredBrush( FSlateNoResource() )
+			.SetTextColor( DefaultForeground )
+			.SetSelectedTextColor( DefaultForeground )
+		);
 	}
 
 	// DestructibleMesh Editor
diff --git a/Engine/Source/Editor/UnrealEd/Private/DataTableEditorUtils.cpp b/Engine/Source/Editor/UnrealEd/Private/DataTableEditorUtils.cpp
--- a/Engine/Source/Editor/UnrealEd/Private/DataTableEditorUtils.cpp
+++ b/Engine/Source/Editor/UnrealEd/Private/DataTableEditorUtils.cpp
@@ -2,6 +2,7 @@
 
 #include "UnrealEd.h"
 #include "DataTableEditorUtils.h"
+#include "DataTableUtils.h"
 #include "K2Node_GetDataTableRow.h"
 #include "Engine/UserDefinedStruct.h"
 
@@ -170,6 +171,68 @@ void FDataTableEditorUtils::BroadcastPostChange(UDataTable* DataTable, EDataTabl
 	FDataTableEditorManager::Get().PostChange(DataTable, Info);
 }
 
+void FDataTableEditorUtils::CacheDataTableForEditing(const UDataTable* DataTable, TArray<FDataTableEditorColumnHeaderDataPtr>& OutAvailableColumns, TArray<FDataTableEditorRowListViewDataPtr>& OutAvailableRows)
+{
+	if (!DataTable || !DataTable->RowStruct)
+	{
+		OutAvailableColumns.Empty();
+		OutAvailableRows.Empty();
+		return;
+	}
+
+	// First build array of properties
+	TArray<const UProperty*> StructProps;
+	for (TFieldIterator<const UProperty> It(DataTable->RowStruct); It; ++It)
+	{
+		const UProperty* Prop = *It;
+		check(Prop);
+		StructProps.Add(Prop);
+	}
+
+	TSharedRef<FSlateFontMeasure> FontMeasure = FSlateApplication::Get().GetRenderer()->GetFontMeasureService();
+	const FTextBlockStyle& CellTextStyle = FEditorStyle::GetWidgetStyle<FTextBlockStyle>("DataTableEditor.CellText");
+	static const float CellPadding = 10.0f;
+
+	// Populate the column data
+	OutAvailableColumns.Empty(StructProps.Num());
+	for (const UProperty* Prop : StructProps)
+	{
+		FDataTableEditorColumnHeaderDataPtr CachedColumnData = MakeShareable(new FDataTableEditorColumnHeaderData());
+		CachedColumnData->ColumnId = Prop->GetFName();
+		CachedColumnData->DisplayName = FText::FromString(DataTableUtils::GetPropertyDisplayName(Prop, Prop->GetName()));
+		CachedColumnData->DesiredColumnWidth = FontMeasure->Measure(CachedColumnData->DisplayName, CellTextStyle.Font).X + CellPadding;
+
+		OutAvailableColumns.Add(CachedColumnData);
+	}
+
+	// Populate the row data
+	OutAvailableRows.Empty(DataTable->RowMap.Num());
+	for (auto RowIt = DataTable->RowMap.CreateConstIterator(); RowIt; ++RowIt)
+	{
+		FDataTableEditorRowListViewDataPtr CachedRowData = MakeShareable(new FDataTableEditorRowListViewData());
+		CachedRowData->RowId = RowIt->Key;
+		CachedRowData->DisplayName = FText::FromName(RowIt->Key);
+
+		CachedRowData->CellData.Reserve(StructProps.Num());
+		{
+			uint8* RowData = RowIt.Value();
+			for (int32 ColumnIndex = 0; ColumnIndex < StructProps.Num(); ++ColumnIndex)
+			{
+				const UProperty* Prop = StructProps[ColumnIndex];
+				FDataTableEditorColumnHeaderDataPtr CachedColumnData = OutAvailableColumns[ColumnIndex];
+
+				const FText CellText = FText::FromString(DataTableUtils::GetPropertyValueAsString(Prop, RowData));
+				CachedRowData->CellData.Add(CellText);
+
+				const float CellWidth = FontMeasure->Measure(CellText, CellTextStyle.Font).X + CellPadding;
+				CachedColumnData->DesiredColumnWidth = FMath::Max(CachedColumnData->DesiredColumnWidth, CellWidth);
+			}
+		}
+
+		OutAvailableRows.Add(CachedRowData);
+	}
+}
+
 TArray<UScriptStruct*> FDataTableEditorUtils::GetPossibleStructs()
 {
 	TArray< UScriptStruct* > RowStructs;
@@ -191,4 +254,4 @@ TArray<UScriptStruct*> FDataTableEditorUtils::GetPossibleStructs()
 		}
 	}
 	return RowStructs;
-}
\ No newline at end of file
+}
diff --git a/Engine/Source/Editor/UnrealEd/Public/DataTableEditorUtils.h b/Engine/Source/Editor/UnrealEd/Public/DataTableEditorUtils.h
--- a/Engine/Source/Editor/UnrealEd/Public/DataTableEditorUtils.h
+++ b/Engine/Source/Editor/UnrealEd/Public/DataTableEditorUtils.h
@@ -4,6 +4,33 @@
 #include "Engine/DataTable.h"
 #include "ListenerManager.h"
 
+struct FDataTableEditorColumnHeaderData
+{
+	/** Unique ID used to identify this column */
+	FName ColumnId;
+
+	/** Display name of this column */
+	FText DisplayName;
+
+	/** The calculated width of this column taking into account the cell data for each row */
+	float DesiredColumnWidth;
+};
+
+struct FDataTableEditorRowListViewData
+{
+	/** Unique ID used to identify this row */
+	FName RowId;
+
+	/** Display name of this row */
+	FText DisplayName;
+
+	/** Array corresponding to each cell in this row */
+	TArray<FText> CellData;
+};
+
+typedef TSharedPtr<FDataTableEditorColumnHeaderData> FDataTableEditorColumnHeaderDataPtr;
+typedef TSharedPtr<FDataTableEditorRowListViewData>  FDataTableEditorRowListViewDataPtr;
+
 struct UNREALED_API FDataTableEditorUtils
 {
 	enum EDataTableChangeInfo
@@ -39,5 +66,7 @@ struct UNREALED_API FDataTableEditorUtils
 	static void BroadcastPreChange(UDataTable* DataTable, EDataTableChangeInfo Info);
 	static void BroadcastPostChange(UDataTable* DataTable, EDataTableChangeInfo Info);
 
+	static void CacheDataTableForEditing(const UDataTable* DataTable, TArray<FDataTableEditorColumnHeaderDataPtr>& OutAvailableColumns, TArray<FDataTableEditorRowListViewDataPtr>& OutAvailableRows);
+
 	static TArray<UScriptStruct*> GetPossibleStructs();
 };
\ No newline at end of file
