Data table undo/redo support

UETOOL-281 - Improve data table workflows
UETOOL-304 - Data table editor doesn't support undo/redo properly

Undo/redo is now supported when adding, removing, moving, or renaming rows, as well as changing the properties of any given data table row.

This also adds highlighting when filtering the data table view, and slightly adjusts the selected row color so it doesn't blend into the background, as well as ensuring the data table is emptied prior to loading (as undo/redo loads on top of an existing object).

diff --git a/Engine/Source/Editor/DataTableEditor/Private/DataTableEditor.cpp b/Engine/Source/Editor/DataTableEditor/Private/DataTableEditor.cpp
--- a/Engine/Source/Editor/DataTableEditor/Private/DataTableEditor.cpp
+++ b/Engine/Source/Editor/DataTableEditor/Private/DataTableEditor.cpp
@@ -81,12 +81,33 @@ FDataTableEditor::FDataTableEditor()
 
 FDataTableEditor::~FDataTableEditor()
 {
+	GEditor->UnregisterForUndo(this);
+
 	if (DataTable.IsValid())
 	{
 		SaveLayoutData();
 	}
 }
 
+void FDataTableEditor::PostUndo(bool bSuccess)
+{
+	HandleUndoRedo();
+}
+
+void FDataTableEditor::PostRedo(bool bSuccess)
+{
+	HandleUndoRedo();
+}
+
+void FDataTableEditor::HandleUndoRedo()
+{
+	if (DataTable.IsValid())
+	{
+		HandlePostChange();
+		CallbackOnDataTableUndoRedo.ExecuteIfBound();
+	}
+}
+
 void FDataTableEditor::PreChange(const class UUserDefinedStruct* Struct, FStructureEditorUtils::EStructureEditorChangeInfo Info)
 {
 }
@@ -95,11 +116,7 @@ void FDataTableEditor::PostChange(const class UUserDefinedStruct* Struct, FStruc
 {
 	if (Struct && DataTable.IsValid() && (DataTable->RowStruct == Struct))
 	{
-		// We need to cache and restore the selection here as RefreshCachedDataTable will re-create the list view items
-		const FName CachedSelection = HighlightedRowName;
-		HighlightedRowName = NAME_None;
-		RefreshCachedDataTable();
-		RestoreCachedSelection(CachedSelection, true/*bUpdateEvenIfValid*/);
+		HandlePostChange();
 	}
 }
 
@@ -112,14 +129,19 @@ void FDataTableEditor::PostChange(const UDataTable* Changed, FDataTableEditorUti
 	FStringAssetReference::InvalidateTag(); // Should be removed after UE-5615 is fixed
 	if (Changed == DataTable.Get())
 	{
-		// We need to cache and restore the selection here as RefreshCachedDataTable will re-create the list view items
-		const FName CachedSelection = HighlightedRowName;
-		HighlightedRowName = NAME_None;
-		RefreshCachedDataTable();
-		RestoreCachedSelection(CachedSelection, true/*bUpdateEvenIfValid*/);
+		HandlePostChange();
 	}
 }
 
+void FDataTableEditor::HandlePostChange()
+{
+	// We need to cache and restore the selection here as RefreshCachedDataTable will re-create the list view items
+	const FName CachedSelection = HighlightedRowName;
+	HighlightedRowName = NAME_None;
+	RefreshCachedDataTable();
+	RestoreCachedSelection(CachedSelection, true/*bUpdateEvenIfValid*/);
+}
+
 void FDataTableEditor::InitDataTableEditor( const EToolkitMode::Type Mode, const TSharedPtr< class IToolkitHost >& InitToolkitHost, UDataTable* Table )
 {
 	TSharedRef<FTabManager::FLayout> StandaloneDefaultLayout = FTabManager::NewLayout( "Standalone_DataTableEditor_Layout" )
@@ -145,6 +167,9 @@ void FDataTableEditor::InitDataTableEditor( const EToolkitMode::Type Mode, const
 	FDataTableEditorModule& DataTableEditorModule = FModuleManager::LoadModuleChecked<FDataTableEditorModule>( "DataTableEditor" );
 	AddMenuExtender(DataTableEditorModule.GetMenuExtensibilityManager()->GetAllExtenders(GetToolkitCommands(), GetEditingObjects()));
 
+	// Support undo/redo
+	GEditor->RegisterForUndo(this);
+
 	// @todo toolkit world centric editing
 	/*// Setup our tool's layout
 	if( IsWorldCentricAssetEditor() )
@@ -279,6 +304,7 @@ TSharedRef<ITableRow> FDataTableEditor::MakeRowNameWidget(FDataTableEditorRowLis
 				SNew(STextBlock)
 				.ColorAndOpacity(this, &FDataTableEditor::GetRowTextColor, InRowDataPtr->RowId)
 				.Text(InRowDataPtr->DisplayName)
+				.HighlightText(this, &FDataTableEditor::GetFilterText)
 			]
 		];
 }
@@ -313,6 +339,7 @@ TSharedRef<SWidget> FDataTableEditor::MakeCellWidget(FDataTableEditorRowListView
 				.TextStyle(FEditorStyle::Get(), "DataTableEditor.CellText")
 				.ColorAndOpacity(this, &FDataTableEditor::GetRowTextColor, InRowDataPtr->RowId)
 				.Text(InRowDataPtr->CellData[ColumnIndex])
+				.HighlightText(this, &FDataTableEditor::GetFilterText)
 				.ToolTipText(FText::Format(LOCTEXT("ColumnRowNameFmt", "{0}: {1}"), AvailableColumns[ColumnIndex]->DisplayName, InRowDataPtr->CellData[ColumnIndex]))
 			];
 	}
@@ -345,8 +372,14 @@ void FDataTableEditor::OnRowSelectionChanged(FDataTableEditorRowListViewDataPtr
 	}
 }
 
-void FDataTableEditor::OnSearchTextChanged(const FText& SearchText)
+FText FDataTableEditor::GetFilterText() const
 {
+	return ActiveFilterText;
+}
+
+void FDataTableEditor::OnFilterTextChanged(const FText& InFilterText)
+{
+	ActiveFilterText = InFilterText;
 	UpdateVisibleRows();
 }
 
@@ -416,8 +449,7 @@ void FDataTableEditor::RefreshCachedDataTable()
 
 void FDataTableEditor::UpdateVisibleRows()
 {
-	const FText& ActiveSearchText = (SearchBox.IsValid()) ? SearchBox->GetText() : FText::GetEmpty();
-	if (ActiveSearchText.IsEmptyOrWhitespace())
+	if (ActiveFilterText.IsEmptyOrWhitespace())
 	{
 		VisibleRows = AvailableRows;
 	}
@@ -425,12 +457,12 @@ void FDataTableEditor::UpdateVisibleRows()
 	{
 		VisibleRows.Empty(AvailableRows.Num());
 
-		const FString& ActiveSearchString = ActiveSearchText.ToString();
+		const FString& ActiveFilterString = ActiveFilterText.ToString();
 		for (const FDataTableEditorRowListViewDataPtr& RowData : AvailableRows)
 		{
 			bool bPassesFilter = false;
 
-			if (RowData->DisplayName.ToString().Contains(ActiveSearchString))
+			if (RowData->DisplayName.ToString().Contains(ActiveFilterString))
 			{
 				bPassesFilter = true;
 			}
@@ -438,7 +470,7 @@ void FDataTableEditor::UpdateVisibleRows()
 			{
 				for (const FText& CellText : RowData->CellData)
 				{
-					if (CellText.ToString().Contains(ActiveSearchString))
+					if (CellText.ToString().Contains(ActiveFilterString))
 					{
 						bPassesFilter = true;
 						break;
@@ -530,8 +562,9 @@ TSharedRef<SVerticalBox> FDataTableEditor::CreateContentBox()
 		+SVerticalBox::Slot()
 		.AutoHeight()
 		[
-			SAssignNew(SearchBox, SSearchBox)
-			.OnTextChanged(this, &FDataTableEditor::OnSearchTextChanged)
+			SNew(SSearchBox)
+			.InitialText(this, &FDataTableEditor::GetFilterText)
+			.OnTextChanged(this, &FDataTableEditor::OnFilterTextChanged)
 		]
 		+SVerticalBox::Slot()
 		[
@@ -586,6 +619,7 @@ TSharedRef<SWidget> FDataTableEditor::CreateRowEditorBox()
 	auto RowEditor = SNew(SRowEditor, DataTable.Get());
 	RowEditor->RowSelectedCallback.BindSP(this, &FDataTableEditor::SetHighlightedRow);
 	CallbackOnRowHighlighted.BindSP(RowEditor, &SRowEditor::SelectRow);
+	CallbackOnDataTableUndoRedo.BindSP(RowEditor, &SRowEditor::HandleUndoRedo);
 	return RowEditor;
 }
 
@@ -616,6 +650,12 @@ TSharedRef<SDockTab> FDataTableEditor::SpawnTab_DataTable( const FSpawnTabArgs&
 
 	DataTable = Cast<UDataTable>(GetEditingObject());
 
+	// Support undo/redo
+	if (DataTable.IsValid())
+	{
+		DataTable->SetFlags(RF_Transactional);
+	}
+
 	LoadLayoutData();
 
 	return SNew(SDockTab)
diff --git a/Engine/Source/Editor/DataTableEditor/Private/DataTableEditor.h b/Engine/Source/Editor/DataTableEditor/Private/DataTableEditor.h
--- a/Engine/Source/Editor/DataTableEditor/Private/DataTableEditor.h
+++ b/Engine/Source/Editor/DataTableEditor/Private/DataTableEditor.h
@@ -4,7 +4,7 @@
 
 #include "IDataTableEditor.h"
 #include "Toolkits/AssetEditorToolkit.h"
-
+#include "EditorUndoClient.h"
 #include "Kismet2/StructureEditorUtils.h"
 #include "DataTableEditorUtils.h"
 
@@ -15,6 +15,7 @@ DECLARE_DELEGATE_OneParam(FOnRowHighlighted, FName /*Row name*/);
 
 /** Viewer/editor for a DataTable */
 class FDataTableEditor : public IDataTableEditor
+	, public FEditorUndoClient
 	, public FStructureEditorUtils::INotifyOnStructChanged
 	, public FDataTableEditorUtils::INotifyOnDataTableChanged
 {
@@ -46,6 +47,11 @@ public:
 	virtual FString GetWorldCentricTabPrefix() const override;
 	virtual FLinearColor GetWorldCentricTabColorScale() const override;
 
+	// FEditorUndoClient
+	virtual void PostUndo(bool bSuccess) override;
+	virtual void PostRedo(bool bSuccess) override;
+	void HandleUndoRedo();
+
 	// INotifyOnStructChanged
 	virtual void PreChange(const class UUserDefinedStruct* Struct, FStructureEditorUtils::EStructureEditorChangeInfo Info) override;
 	virtual void PostChange(const class UUserDefinedStruct* Struct, FStructureEditorUtils::EStructureEditorChangeInfo Info) override;
@@ -54,6 +60,8 @@ public:
 	virtual void PreChange(const UDataTable* Changed, FDataTableEditorUtils::EDataTableChangeInfo Info) override;
 	virtual void PostChange(const UDataTable* Changed, FDataTableEditorUtils::EDataTableChangeInfo Info) override;
 
+	void HandlePostChange();
+
 	void SetHighlightedRow(FName Name);
 
 private:
@@ -64,7 +72,9 @@ private:
 
 	void RestoreCachedSelection(const FName InCachedSelection, const bool bUpdateEvenIfValid = false);
 
-	void OnSearchTextChanged(const FText& SearchText);
+	FText GetFilterText() const;
+
+	void OnFilterTextChanged(const FText& InFilterText);
 
 	FSlateColor GetRowTextColor(FName RowName) const;
 
@@ -144,19 +154,22 @@ private:
 	/** Widths of data table cell columns */
 	TArray<FColumnWidth> ColumnWidths;
 
-	/** Search box */
-	TSharedPtr<SSearchBox> SearchBox;
-
 	/* The DataTable that is active in the editor */
 	TAssetPtr<UDataTable> DataTable;
 
 	/** The layout data for the currently loaded data table */
 	TSharedPtr<FJsonObject> LayoutData;
 
+	/** The name of the currently selected row */
 	FName HighlightedRowName;
 
+	/** The current filter text applied to the data table */
+	FText ActiveFilterText;
+
 	FOnRowHighlighted CallbackOnRowHighlighted;
 
+	FSimpleDelegate CallbackOnDataTableUndoRedo;
+
 	/**	The tab id for the data table tab */
 	static const FName DataTableTabId;
 
diff --git a/Engine/Source/Editor/DataTableEditor/Private/SRowEditor.cpp b/Engine/Source/Editor/DataTableEditor/Private/SRowEditor.cpp
--- a/Engine/Source/Editor/DataTableEditor/Private/SRowEditor.cpp
+++ b/Engine/Source/Editor/DataTableEditor/Private/SRowEditor.cpp
@@ -59,7 +59,8 @@ public:
 	}
 };
 
-SRowEditor::SRowEditor() : SCompoundWidget()
+SRowEditor::SRowEditor() 
+	: SCompoundWidget()
 {
 }
 
@@ -67,22 +68,21 @@ SRowEditor::~SRowEditor()
 {
 }
 
-void SRowEditor::OnFinishedChangingProperties(const FPropertyChangedEvent& PropertyChangedEvent)
+void SRowEditor::NotifyPreChange( UProperty* PropertyAboutToChange )
 {
-	check(PropertyChangedEvent.MemberProperty
-		&& PropertyChangedEvent.MemberProperty->GetOwnerStruct()
-		&& (PropertyChangedEvent.MemberProperty->GetOwnerStruct() == GetScriptStruct()));
+	check(DataTable.IsValid());
+	DataTable->Modify();
 
-	// TODO:
-	// FNotifyHook won't be called, because the struct has no "property-path" from DT object. THe struct is seen as detached.
-	// But the notification must be called somewhere.
+	FDataTableEditorUtils::BroadcastPreChange(DataTable.Get(), FDataTableEditorUtils::EDataTableChangeInfo::RowData);
+}
 
-	FDataTableEditorUtils::BroadcastPostChange(DataTable.Get(), FDataTableEditorUtils::EDataTableChangeInfo::RowDataPostChangeOnly);
+void SRowEditor::NotifyPostChange( const FPropertyChangedEvent& PropertyChangedEvent, UProperty* PropertyThatChanged )
+{
+	check(DataTable.IsValid());
 
-	if (DataTable.IsValid())
-	{
-		DataTable->MarkPackageDirty();
-	}
+	FDataTableEditorUtils::BroadcastPostChange(DataTable.Get(), FDataTableEditorUtils::EDataTableChangeInfo::RowData);
+
+	DataTable->MarkPackageDirty();
 }
 
 void SRowEditor::PreChange(const class UUserDefinedStruct* Struct, FStructureEditorUtils::EStructureEditorChangeInfo Info)
@@ -254,6 +254,12 @@ void SRowEditor::SelectRow(FName InName)
 	OnSelectionChanged(NewSelectedName, ESelectInfo::Direct);
 }
 
+void SRowEditor::HandleUndoRedo()
+{
+	RefreshNameList();
+	Restore();
+}
+
 FReply SRowEditor::OnAddClicked()
 {
 	if (DataTable.IsValid())
@@ -330,9 +336,9 @@ void SRowEditor::Construct(const FArguments& InArgs, UDataTable* Changed)
 		ViewArgs.bAllowSearch = false;
 		ViewArgs.bHideSelectionTip = false;
 		ViewArgs.bShowActorLabel = false;
+		ViewArgs.NotifyHook = this;
 
 		StructureDetailsView = PropertyModule.CreateStructureDetailView(ViewArgs, CurrentRow, true/*bShowObjects*/, LOCTEXT("RowValue", "Row Value"));
-		StructureDetailsView->GetOnFinishedChangingPropertiesDelegate().AddSP(this, &SRowEditor::OnFinishedChangingProperties);
 	}
 
 	RefreshNameList();
diff --git a/Engine/Source/Editor/DataTableEditor/Private/SRowEditor.h b/Engine/Source/Editor/DataTableEditor/Private/SRowEditor.h
--- a/Engine/Source/Editor/DataTableEditor/Private/SRowEditor.h
+++ b/Engine/Source/Editor/DataTableEditor/Private/SRowEditor.h
@@ -9,6 +9,7 @@ DECLARE_DELEGATE_OneParam(FOnRowModified, FName /*Row name*/);
 DECLARE_DELEGATE_OneParam(FOnRowSelected, FName /*Row name*/);
 
 class SRowEditor : public SCompoundWidget
+	, public FNotifyHook
 	, public FStructureEditorUtils::INotifyOnStructChanged
 	, public FDataTableEditorUtils::INotifyOnDataTableChanged
 {
@@ -19,6 +20,10 @@ public:
 	SRowEditor();
 	virtual ~SRowEditor();
 
+	// FNotifyHook
+	virtual void NotifyPreChange( UProperty* PropertyAboutToChange ) override;
+	virtual void NotifyPostChange( const FPropertyChangedEvent& PropertyChangedEvent, UProperty* PropertyThatChanged ) override;
+
 	// INotifyOnStructChanged
 	virtual void PreChange(const class UUserDefinedStruct* Struct, FStructureEditorUtils::EStructureEditorChangeInfo Info) override;
 	virtual void PostChange(const class UUserDefinedStruct* Struct, FStructureEditorUtils::EStructureEditorChangeInfo Info) override;
@@ -51,8 +56,6 @@ private:
 	TSharedRef<SWidget> OnGenerateWidget(TSharedPtr<FName> InItem);
 	void OnSelectionChanged(TSharedPtr<FName> InItem, ESelectInfo::Type InSeletionInfo);
 
-	void OnFinishedChangingProperties(const struct FPropertyChangedEvent& PropertyChangedEvent);
-
 	FReply OnAddClicked();
 	FReply OnRemoveClicked();
 	FReply OnMoveRowClicked(FDataTableEditorUtils::ERowMoveDirection MoveDirection);
@@ -64,4 +67,6 @@ public:
 	void Construct(const FArguments& InArgs, UDataTable* Changed);
 
 	void SelectRow(FName Name);
-};
\ No newline at end of file
+
+	void HandleUndoRedo();
+};
diff --git a/Engine/Source/Editor/EditorStyle/Private/SlateEditorStyle.cpp b/Engine/Source/Editor/EditorStyle/Private/SlateEditorStyle.cpp
--- a/Engine/Source/Editor/EditorStyle/Private/SlateEditorStyle.cpp
+++ b/Engine/Source/Editor/EditorStyle/Private/SlateEditorStyle.cpp
@@ -5673,10 +5673,10 @@ void FSlateEditorStyle::FStyle::SetupToolkitStyles()
 			.SetEvenRowBackgroundHoveredBrush(IMAGE_BRUSH("PropertyView/DetailCategoryMiddle_Hovered", FVector2D(16, 16), FLinearColor(0.5f, 0.5f, 0.5f)))
 			.SetOddRowBackgroundBrush(IMAGE_BRUSH("PropertyView/DetailCategoryMiddle", FVector2D(16, 16), FLinearColor(0.2f, 0.2f, 0.2f)))
 			.SetOddRowBackgroundHoveredBrush(IMAGE_BRUSH("PropertyView/DetailCategoryMiddle_Hovered", FVector2D(16, 16), FLinearColor(0.2f, 0.2f, 0.2f)))
-			.SetActiveBrush( FSlateNoResource() )
-			.SetActiveHoveredBrush( FSlateNoResource() )
-			.SetInactiveBrush( FSlateNoResource() )
-			.SetInactiveHoveredBrush( FSlateNoResource() )
+			.SetActiveBrush( IMAGE_BRUSH( "Common/Selection", Icon8x8, FLinearColor(0.075f, 0.075f, 0.075f) ) )
+			.SetActiveHoveredBrush( IMAGE_BRUSH( "Common/Selection", Icon8x8, FLinearColor(0.075f, 0.075f, 0.075f) ) )
+			.SetInactiveBrush( IMAGE_BRUSH( "Common/Selection", Icon8x8, FLinearColor(0.075f, 0.075f, 0.075f) ) )
+			.SetInactiveHoveredBrush( IMAGE_BRUSH( "Common/Selection", Icon8x8, FLinearColor(0.075f, 0.075f, 0.075f) ) )
 			.SetTextColor( DefaultForeground )
 			.SetSelectedTextColor( DefaultForeground )
 		);
diff --git a/Engine/Source/Editor/UnrealEd/Private/DataTableEditorUtils.cpp b/Engine/Source/Editor/UnrealEd/Private/DataTableEditorUtils.cpp
--- a/Engine/Source/Editor/UnrealEd/Private/DataTableEditorUtils.cpp
+++ b/Engine/Source/Editor/UnrealEd/Private/DataTableEditorUtils.cpp
@@ -3,9 +3,12 @@
 #include "UnrealEd.h"
 #include "DataTableEditorUtils.h"
 #include "DataTableUtils.h"
+#include "ScopedTransaction.h"
 #include "K2Node_GetDataTableRow.h"
 #include "Engine/UserDefinedStruct.h"
 
+#define LOCTEXT_NAMESPACE "DataTableEditorUtils"
+
 FDataTableEditorUtils::FDataTableEditorManager& FDataTableEditorUtils::FDataTableEditorManager::Get()
 {
 	static TSharedRef< FDataTableEditorManager > EditorManager(new FDataTableEditorManager());
@@ -17,6 +20,8 @@ bool FDataTableEditorUtils::RemoveRow(UDataTable* DataTable, FName Name)
 	bool bResult = false;
 	if (DataTable && DataTable->RowStruct)
 	{
+		const FScopedTransaction Transaction(LOCTEXT("RemoveDataTableRow", "Remove Data Table Row"));
+
 		BroadcastPreChange(DataTable, EDataTableChangeInfo::RowList);
 		DataTable->Modify();
 		uint8* RowData = NULL;
@@ -41,6 +46,9 @@ uint8* FDataTableEditorUtils::AddRow(UDataTable* DataTable, FName RowName)
 	{
 		return NULL;
 	}
+
+	const FScopedTransaction Transaction(LOCTEXT("AddDataTableRow", "Add Data Table Row"));
+
 	BroadcastPreChange(DataTable, EDataTableChangeInfo::RowList);
 	DataTable->Modify();
 	// Allocate data to store information, using UScriptStruct to know its size
@@ -64,6 +72,8 @@ bool FDataTableEditorUtils::RenameRow(UDataTable* DataTable, FName OldName, FNam
 	bool bResult = false;
 	if (DataTable)
 	{
+		const FScopedTransaction Transaction(LOCTEXT("RenameDataTableRow", "Rename Data Table Row"));
+
 		BroadcastPreChange(DataTable, EDataTableChangeInfo::RowList);
 		DataTable->Modify();
 
@@ -138,6 +148,8 @@ bool FDataTableEditorUtils::MoveRow(UDataTable* DataTable, FName RowName, ERowMo
 		NamesToNewIndex.Add(OrderedRowNames[NameIndex], NameIndex);
 	}
 
+	const FScopedTransaction Transaction(LOCTEXT("MoveDataTableRow", "Move Data Table Row"));
+
 	BroadcastPreChange(DataTable, EDataTableChangeInfo::RowList);
 	DataTable->Modify();
 
@@ -255,3 +267,5 @@ TArray<UScriptStruct*> FDataTableEditorUtils::GetPossibleStructs()
 	}
 	return RowStructs;
 }
+
+#undef LOCTEXT_NAMESPACE
diff --git a/Engine/Source/Editor/UnrealEd/Public/DataTableEditorUtils.h b/Engine/Source/Editor/UnrealEd/Public/DataTableEditorUtils.h
--- a/Engine/Source/Editor/UnrealEd/Public/DataTableEditorUtils.h
+++ b/Engine/Source/Editor/UnrealEd/Public/DataTableEditorUtils.h
@@ -33,9 +33,11 @@ typedef TSharedPtr<FDataTableEditorRowListViewData>  FDataTableEditorRowListView
 
 struct UNREALED_API FDataTableEditorUtils
 {
-	enum EDataTableChangeInfo
+	enum class EDataTableChangeInfo
 	{
-		RowDataPostChangeOnly,
+		/** The data corresponding to a single row has been changed */
+		RowData,
+		/** The data corresponding to the entire list of rows has been changed */
 		RowList,
 	};
 
diff --git a/Engine/Source/Runtime/Engine/Private/DataTable.cpp b/Engine/Source/Runtime/Engine/Private/DataTable.cpp
--- a/Engine/Source/Runtime/Engine/Private/DataTable.cpp
+++ b/Engine/Source/Runtime/Engine/Private/DataTable.cpp
@@ -82,6 +82,7 @@ void UDataTable::Serialize( FArchive& Ar )
 
 	if(Ar.IsLoading())
 	{
+		EmptyTable();
 		LoadStructData(Ar);
 	}
 	else if(Ar.IsSaving())
