Fixed order of initialisation bug by using GetEditingObject() instead

The DataTable pointer was only set when the data table tab was opened, so if you re-ordered your tab so that the row editor was created first, the row editor would be using a null data table pointer and the editor wouldn't work correctly.

diff --git a/Engine/Source/Editor/CurveTableEditor/Private/CurveTableEditor.cpp b/Engine/Source/Editor/CurveTableEditor/Private/CurveTableEditor.cpp
--- a/Engine/Source/Editor/CurveTableEditor/Private/CurveTableEditor.cpp
+++ b/Engine/Source/Editor/CurveTableEditor/Private/CurveTableEditor.cpp
@@ -130,6 +130,11 @@ FLinearColor FCurveTableEditor::GetWorldCentricTabColorScale() const
 }
 
 
+const UCurveTable* FCurveTableEditor::GetCurveTable() const
+{
+	return Cast<const UCurveTable>(GetEditingObject());
+}
+
 TSharedRef<SDockTab> FCurveTableEditor::SpawnTab_CurveTable( const FSpawnTabArgs& Args )
 {
 	check( Args.GetTabId().TabType == CurveTableTabId );
@@ -258,7 +263,7 @@ void FCurveTableEditor::CacheDataTableForEditing()
 {
 	RowNameColumnWidth = 10.0f;
 
-	const UCurveTable* Table = Cast<const UCurveTable>(GetEditingObject());
+	const UCurveTable* Table = GetCurveTable();
 	if (!Table || Table->RowMap.Num() == 0)
 	{
 		AvailableColumns.Empty();
@@ -408,7 +413,7 @@ FOptionalSize FCurveTableEditor::GetRowNameColumnWidth() const
 
 void FCurveTableEditor::OnPostReimport(UObject* InObject, bool)
 {
-	const UCurveTable* Table = Cast<const UCurveTable>(GetEditingObject());
+	const UCurveTable* Table = GetCurveTable();
 	if (Table && Table == InObject)
 	{
 		RefreshCachedCurveTable();
diff --git a/Engine/Source/Editor/CurveTableEditor/Private/CurveTableEditor.h b/Engine/Source/Editor/CurveTableEditor/Private/CurveTableEditor.h
--- a/Engine/Source/Editor/CurveTableEditor/Private/CurveTableEditor.h
+++ b/Engine/Source/Editor/CurveTableEditor/Private/CurveTableEditor.h
@@ -62,8 +62,8 @@ public:
 	virtual FString GetWorldCentricTabPrefix() const override;
 	virtual FLinearColor GetWorldCentricTabColorScale() const override;
 
-	/** Creates the grid widget used to display the data */
-	//TSharedRef<SUniformGridPanel> CreateGridPanel() const;
+	/** Get the curve table being edited */
+	const UCurveTable* GetCurveTable() const;
 
 	/**	Spawns the tab with the curve table inside */
 	TSharedRef<SDockTab> SpawnTab_CurveTable( const FSpawnTabArgs& Args );
diff --git a/Engine/Source/Editor/DataTableEditor/Private/DataTableEditor.cpp b/Engine/Source/Editor/DataTableEditor/Private/DataTableEditor.cpp
--- a/Engine/Source/Editor/DataTableEditor/Private/DataTableEditor.cpp
+++ b/Engine/Source/Editor/DataTableEditor/Private/DataTableEditor.cpp
@@ -83,7 +83,8 @@ FDataTableEditor::~FDataTableEditor()
 {
 	GEditor->UnregisterForUndo(this);
 
-	if (DataTable.IsValid())
+	const UDataTable* Table = GetDataTable();
+	if (Table)
 	{
 		SaveLayoutData();
 	}
@@ -101,7 +102,8 @@ void FDataTableEditor::PostRedo(bool bSuccess)
 
 void FDataTableEditor::HandleUndoRedo()
 {
-	if (DataTable.IsValid())
+	const UDataTable* Table = GetDataTable();
+	if (Table)
 	{
 		HandlePostChange();
 		CallbackOnDataTableUndoRedo.ExecuteIfBound();
@@ -114,7 +116,8 @@ void FDataTableEditor::PreChange(const class UUserDefinedStruct* Struct, FStruct
 
 void FDataTableEditor::PostChange(const class UUserDefinedStruct* Struct, FStructureEditorUtils::EStructureEditorChangeInfo Info)
 {
-	if (Struct && DataTable.IsValid() && (DataTable->RowStruct == Struct))
+	const UDataTable* Table = GetDataTable();
+	if (Struct && Table && (Table->RowStruct == Struct))
 	{
 		HandlePostChange();
 	}
@@ -126,13 +129,18 @@ void FDataTableEditor::PreChange(const UDataTable* Changed, FDataTableEditorUtil
 
 void FDataTableEditor::PostChange(const UDataTable* Changed, FDataTableEditorUtils::EDataTableChangeInfo Info)
 {
-	FStringAssetReference::InvalidateTag(); // Should be removed after UE-5615 is fixed
-	if (Changed == DataTable.Get())
+	const UDataTable* Table = GetDataTable();
+	if (Changed == Table)
 	{
 		HandlePostChange();
 	}
 }
 
+const UDataTable* FDataTableEditor::GetDataTable() const
+{
+	return Cast<const UDataTable>(GetEditingObject());
+}
+
 void FDataTableEditor::HandlePostChange()
 {
 	// We need to cache and restore the selection here as RefreshCachedDataTable will re-create the list view items
@@ -260,12 +268,13 @@ void FDataTableEditor::LoadLayoutData()
 {
 	LayoutData.Reset();
 
-	if (!DataTable.IsValid())
+	const UDataTable* Table = GetDataTable();
+	if (!Table)
 	{
 		return;
 	}
 
-	const FString LayoutDataFilename = FPaths::GameSavedDir() / TEXT("AssetData") / TEXT("DataTableEditorLayout") / DataTable->GetName() + TEXT(".json");
+	const FString LayoutDataFilename = FPaths::GameSavedDir() / TEXT("AssetData") / TEXT("DataTableEditorLayout") / Table->GetName() + TEXT(".json");
 
 	FString JsonText;
 	if (FFileHelper::LoadFileToString(JsonText, *LayoutDataFilename))
@@ -277,12 +286,13 @@ void FDataTableEditor::LoadLayoutData()
 
 void FDataTableEditor::SaveLayoutData()
 {
-	if (!DataTable.IsValid() || !LayoutData.IsValid())
+	const UDataTable* Table = GetDataTable();
+	if (!Table || !LayoutData.IsValid())
 	{
 		return;
 	}
 
-	const FString LayoutDataFilename = FPaths::GameSavedDir() / TEXT("AssetData") / TEXT("DataTableEditorLayout") / DataTable->GetName() + TEXT(".json");
+	const FString LayoutDataFilename = FPaths::GameSavedDir() / TEXT("AssetData") / TEXT("DataTableEditorLayout") / Table->GetName() + TEXT(".json");
 
 	FString JsonText;
 	TSharedRef< TJsonWriter< TCHAR, TPrettyJsonPrintPolicy<TCHAR> > > JsonWriter = TJsonWriterFactory< TCHAR, TPrettyJsonPrintPolicy<TCHAR> >::Create(&JsonText);
@@ -385,7 +395,8 @@ void FDataTableEditor::OnFilterTextChanged(const FText& InFilterText)
 
 void FDataTableEditor::RefreshCachedDataTable()
 {
-	FDataTableEditorUtils::CacheDataTableForEditing(DataTable.Get(), AvailableColumns, AvailableRows);
+	const UDataTable* Table = GetDataTable();
+	FDataTableEditorUtils::CacheDataTableForEditing(Table, AvailableColumns, AvailableRows);
 
 	// Update the desired width of the row names column
 	// This prevents it growing or shrinking as you scroll the list view
@@ -616,7 +627,15 @@ TSharedRef<SVerticalBox> FDataTableEditor::CreateContentBox()
 
 TSharedRef<SWidget> FDataTableEditor::CreateRowEditorBox()
 {
-	auto RowEditor = SNew(SRowEditor, DataTable.Get());
+	UDataTable* Table = Cast<UDataTable>(GetEditingObject());
+
+	// Support undo/redo
+	if (Table)
+	{
+		Table->SetFlags(RF_Transactional);
+	}
+
+	auto RowEditor = SNew(SRowEditor, Table);
 	RowEditor->RowSelectedCallback.BindSP(this, &FDataTableEditor::SetHighlightedRow);
 	CallbackOnRowHighlighted.BindSP(RowEditor, &SRowEditor::SelectRow);
 	CallbackOnDataTableUndoRedo.BindSP(RowEditor, &SRowEditor::HandleUndoRedo);
@@ -648,12 +667,12 @@ TSharedRef<SDockTab> FDataTableEditor::SpawnTab_DataTable( const FSpawnTabArgs&
 {
 	check( Args.GetTabId().TabType == DataTableTabId );
 
-	DataTable = Cast<UDataTable>(GetEditingObject());
+	UDataTable* Table = Cast<UDataTable>(GetEditingObject());
 
 	// Support undo/redo
-	if (DataTable.IsValid())
+	if (Table)
 	{
-		DataTable->SetFlags(RF_Transactional);
+		Table->SetFlags(RF_Transactional);
 	}
 
 	LoadLayoutData();
diff --git a/Engine/Source/Editor/DataTableEditor/Private/DataTableEditor.h b/Engine/Source/Editor/DataTableEditor/Private/DataTableEditor.h
--- a/Engine/Source/Editor/DataTableEditor/Private/DataTableEditor.h
+++ b/Engine/Source/Editor/DataTableEditor/Private/DataTableEditor.h
@@ -60,6 +60,9 @@ public:
 	virtual void PreChange(const UDataTable* Changed, FDataTableEditorUtils::EDataTableChangeInfo Info) override;
 	virtual void PostChange(const UDataTable* Changed, FDataTableEditorUtils::EDataTableChangeInfo Info) override;
 
+	/** Get the data table being edited */
+	const UDataTable* GetDataTable() const;
+
 	void HandlePostChange();
 
 	void SetHighlightedRow(FName Name);
@@ -154,9 +157,6 @@ private:
 	/** Widths of data table cell columns */
 	TArray<FColumnWidth> ColumnWidths;
 
-	/* The DataTable that is active in the editor */
-	TAssetPtr<UDataTable> DataTable;
-
 	/** The layout data for the currently loaded data table */
 	TSharedPtr<FJsonObject> LayoutData;
 
