Fixed a crash when changing a struct that was being used by an open Data Table editor

UE-17680 - Adding a variable to a Data Table struct causes crash

diff --git a/Engine/Source/Editor/DataTableEditor/Private/DataTableEditor.cpp b/Engine/Source/Editor/DataTableEditor/Private/DataTableEditor.cpp
--- a/Engine/Source/Editor/DataTableEditor/Private/DataTableEditor.cpp
+++ b/Engine/Source/Editor/DataTableEditor/Private/DataTableEditor.cpp
@@ -340,7 +340,7 @@ TSharedRef<SWidget> FDataTableEditor::MakeCellWidget(FDataTableEditorRowListView
 	}
 
 	// Valid column ID?
-	if (AvailableColumns.IsValidIndex(ColumnIndex))
+	if (AvailableColumns.IsValidIndex(ColumnIndex) && InRowDataPtr->CellData.IsValidIndex(ColumnIndex))
 	{
 		return SNew(SBox)
 			.Padding(FMargin(4, 2, 4, 2))
