Don't allow users to try to rename rows to emtpy strings, whitespace only, or names that have invalid characters.

diff --git a/Engine/Source/Editor/DataTableEditor/Private/SRowEditor.cpp b/Engine/Source/Editor/DataTableEditor/Private/SRowEditor.cpp
--- a/Engine/Source/Editor/DataTableEditor/Private/SRowEditor.cpp
+++ b/Engine/Source/Editor/DataTableEditor/Private/SRowEditor.cpp
@@ -8,6 +8,7 @@
 #include "Widgets/Input/SButton.h"
 #include "EditorStyleSet.h"
 #include "UObject/StructOnScope.h"
+#include "Misc/MessageDialog.h"
 
 #include "PropertyEditorModule.h"
 #include "IStructureDetailsView.h"
@@ -336,6 +337,14 @@ void SRowEditor::OnRowRenamed(const FText& Text, ETextCommit::Type CommitType)
 {
 	if (!GetCurrentNameAsText().EqualTo(Text) && DataTable.IsValid())
 	{
+		if (Text.IsEmptyOrWhitespace() || !FName::IsValidXName(Text.ToString(), INVALID_NAME_CHARACTERS))
+		{
+			// popup an error dialog here
+			const FText Message = FText::Format(LOCTEXT("InvalidRowName", "'{0}' is not a valid row name"), Text);
+			FMessageDialog::Open(EAppMsgType::Ok, Message);
+
+			return;
+		}
 		const FName NewName = DataTableUtils::MakeValidName(Text.ToString());
 		for (auto Name : CachedRowNames)
 		{
