Fix for a warning popping up twice for one user action

diff --git a/Engine/Source/Editor/DataTableEditor/Private/SRowEditor.cpp b/Engine/Source/Editor/DataTableEditor/Private/SRowEditor.cpp
--- a/Engine/Source/Editor/DataTableEditor/Private/SRowEditor.cpp
+++ b/Engine/Source/Editor/DataTableEditor/Private/SRowEditor.cpp
@@ -339,10 +339,13 @@ void SRowEditor::OnRowRenamed(const FText& Text, ETextCommit::Type CommitType)
 	{
 		if (Text.IsEmptyOrWhitespace() || !FName::IsValidXName(Text.ToString(), INVALID_NAME_CHARACTERS))
 		{
-			// popup an error dialog here
-			const FText Message = FText::Format(LOCTEXT("InvalidRowName", "'{0}' is not a valid row name"), Text);
-			FMessageDialog::Open(EAppMsgType::Ok, Message);
-
+			// Only pop up the error dialog if the rename was caused by the user's action
+			if ((CommitType == ETextCommit::OnEnter) || (CommitType == ETextCommit::OnUserMovedFocus ))
+			{
+				// popup an error dialog here
+				const FText Message = FText::Format(LOCTEXT("InvalidRowName", "'{0}' is not a valid row name"), Text);
+				FMessageDialog::Open(EAppMsgType::Ok, Message);
+			}
 			return;
 		}
 		const FName NewName = DataTableUtils::MakeValidName(Text.ToString());
