Change shared pointers stored on SDataTableListViewRowName to weak pointers, so datatable editor is cleaned up & can be reopened
#rnx

diff --git a/Engine/Source/Editor/DataTableEditor/Private/SDataTableListViewRowName.cpp b/Engine/Source/Editor/DataTableEditor/Private/SDataTableListViewRowName.cpp
--- a/Engine/Source/Editor/DataTableEditor/Private/SDataTableListViewRowName.cpp
+++ b/Engine/Source/Editor/DataTableEditor/Private/SDataTableListViewRowName.cpp
@@ -24,9 +24,9 @@ void SDataTableListViewRowName::Construct(const FArguments& InArgs, const TShare
 				.HeightOverride(RowDataPtr->DesiredRowHeight)
 				[
 					SNew(STextBlock)
-					.ColorAndOpacity(DataTableEditor.Get(), &FDataTableEditor::GetRowTextColor, RowDataPtr->RowId)
+					.ColorAndOpacity(DataTableEditor.Pin().Get(), &FDataTableEditor::GetRowTextColor, RowDataPtr->RowId)
 					.Text(RowDataPtr->DisplayName)
-					.HighlightText(DataTableEditor.Get(), &FDataTableEditor::GetFilterText)
+					.HighlightText(DataTableEditor.Pin().Get(), &FDataTableEditor::GetFilterText)
 				]
 			]
 		],
@@ -54,11 +54,14 @@ void SDataTableListViewRowName::OnSearchForReferences()
 {
 	if (DataTableEditor.IsValid() && RowDataPtr.IsValid())
 	{
-		UDataTable* SourceDataTable = const_cast<UDataTable*>(DataTableEditor->GetDataTable());
-		
-		TArray<FAssetIdentifier> AssetIdentifiers;
-		AssetIdentifiers.Add(FAssetIdentifier(SourceDataTable, RowDataPtr->RowId));
+		if (FDataTableEditor* DataTableEditorPtr = DataTableEditor.Pin().Get())
+		{
+			UDataTable* SourceDataTable = const_cast<UDataTable*>(DataTableEditorPtr->GetDataTable());
+
+			TArray<FAssetIdentifier> AssetIdentifiers;
+			AssetIdentifiers.Add(FAssetIdentifier(SourceDataTable, RowDataPtr->RowId));
 
-		FEditorDelegates::OnOpenReferenceViewer.Broadcast(AssetIdentifiers);
+			FEditorDelegates::OnOpenReferenceViewer.Broadcast(AssetIdentifiers);
+		}
 	}
 }
diff --git a/Engine/Source/Editor/DataTableEditor/Private/SDataTableListViewRowName.h b/Engine/Source/Editor/DataTableEditor/Private/SDataTableListViewRowName.h
--- a/Engine/Source/Editor/DataTableEditor/Private/SDataTableListViewRowName.h
+++ b/Engine/Source/Editor/DataTableEditor/Private/SDataTableListViewRowName.h
@@ -35,5 +35,5 @@ private:
 
 	void OnSearchForReferences();
 	FDataTableEditorRowListViewDataPtr RowDataPtr;
-	TSharedPtr<FDataTableEditor> DataTableEditor;
+	TWeakPtr<FDataTableEditor> DataTableEditor;
 };
\ No newline at end of file
