Added a row type display to the data table editor similar to the parent type in the blueprint editor.  Handles navigation to user defined structs and c++ structs

diff --git a/Engine/Source/Editor/DataTableEditor/Private/DataTableEditor.cpp b/Engine/Source/Editor/DataTableEditor/Private/DataTableEditor.cpp
--- a/Engine/Source/Editor/DataTableEditor/Private/DataTableEditor.cpp
+++ b/Engine/Source/Editor/DataTableEditor/Private/DataTableEditor.cpp
@@ -22,6 +22,10 @@
 #include "Widgets/SToolTip.h"
 #include "Widgets/Text/STextBlock.h"
 #include "Widgets/Views/SListView.h"
+#include "Widgets/Input/SButton.h"
+#include "Widgets/Images/SImage.h"
+#include "Widgets/Input/SHyperlink.h"
+#include "SourceCodeNavigation.h"
 
 #define LOCTEXT_NAMESPACE "DataTableEditor"
 
@@ -228,6 +232,8 @@ void FDataTableEditor::InitDataTableEditor( const EToolkitMode::Type Mode, const
 	FDataTableEditorModule& DataTableEditorModule = FModuleManager::LoadModuleChecked<FDataTableEditorModule>( "DataTableEditor" );
 	AddMenuExtender(DataTableEditorModule.GetMenuExtensibilityManager()->GetAllExtenders(GetToolkitCommands(), GetEditingObjects()));
 
+	RegenerateMenusAndToolbars();
+
 	// Support undo/redo
 	GEditor->RegisterForUndo(this);
 
@@ -463,6 +469,90 @@ void FDataTableEditor::OnFilterTextChanged(const FText& InFilterText)
 	UpdateVisibleRows();
 }
 
+void FDataTableEditor::PostRegenerateMenusAndToolbars()
+{
+	const UDataTable* DataTable = GetDataTable();
+
+	if (DataTable)
+	{
+		const UUserDefinedStruct* UDS = Cast<const UUserDefinedStruct>(DataTable->GetRowStruct());
+
+		// build and attach the menu overlay
+		TSharedRef<SHorizontalBox> MenuOverlayBox = SNew(SHorizontalBox)
+			+ SHorizontalBox::Slot()
+			.AutoWidth()
+			.VAlign(VAlign_Center)
+			[
+				SNew(STextBlock)
+				.ColorAndOpacity(FSlateColor::UseSubduedForeground())
+				.ShadowOffset(FVector2D::UnitVector)
+				.Text(LOCTEXT("DataTableEditor_RowStructType", "Row Type: "))
+			]
+			+ SHorizontalBox::Slot()
+			.AutoWidth()
+			.VAlign(VAlign_Center)
+			[
+				SNew(STextBlock)
+				.ShadowOffset(FVector2D::UnitVector)
+				.Text(FText::FromName(DataTable->GetRowStructName()))
+				.ToolTipText(LOCTEXT("DataTableRowToolTip", "The struct used for each row in this data table"))
+				.Visibility(UDS ? EVisibility::Visible : EVisibility::Collapsed)
+			]
+			+ SHorizontalBox::Slot()
+			.AutoWidth()
+			[
+				SNew(SButton)
+				.VAlign(VAlign_Center)
+				.ButtonStyle(FEditorStyle::Get(), "HoverHintOnly")
+				.OnClicked(this, &FDataTableEditor::OnFindRowInContentBrowserClicked)
+				.Visibility(UDS ? EVisibility::Visible : EVisibility::Collapsed)
+				.ToolTipText(LOCTEXT("FindRowInCBToolTip", "Find row in Content Browser"))
+				.ContentPadding(4.0f)
+				.ForegroundColor(FSlateColor::UseForeground())
+				[
+					SNew(SImage)
+					.Image(FEditorStyle::GetBrush("PropertyWindow.Button_Browse"))
+				]
+			]
+			+SHorizontalBox::Slot()
+			.AutoWidth()
+			.VAlign(VAlign_Center)
+			.Padding(0.0f, 0.0f, 8.0f, 0.0f)
+			[
+				SNew(SHyperlink)
+				.Style(FEditorStyle::Get(), "Common.GotoNativeCodeHyperlink")
+				.Visibility(!UDS ? EVisibility::Visible : EVisibility::Collapsed)
+				.OnNavigate(this, &FDataTableEditor::OnNavigateToDataTableRowCode)
+				.Text(FText::FromName(DataTable->GetRowStructName()))
+				.ToolTipText(FText::Format(LOCTEXT("GoToCode_ToolTip", "Click to open this source file in {0}"), FSourceCodeNavigation::GetSelectedSourceCodeIDE()))
+			];
+	
+		SetMenuOverlay(MenuOverlayBox);
+	}
+}
+
+FReply FDataTableEditor::OnFindRowInContentBrowserClicked()
+{
+	const UDataTable* DataTable = GetDataTable();
+	if(DataTable)
+	{
+		TArray<FAssetData> ObjectsToSync;
+		ObjectsToSync.Add(FAssetData(DataTable->GetRowStruct()));
+		GEditor->SyncBrowserToObjects(ObjectsToSync);
+	}
+
+	return FReply::Handled();
+}
+
+void FDataTableEditor::OnNavigateToDataTableRowCode()
+{
+	const UDataTable* DataTable = GetDataTable();
+	if (DataTable && FSourceCodeNavigation::NavigateToStruct(DataTable->GetRowStruct()))
+	{
+		FSourceCodeNavigation::NavigateToStruct(DataTable->GetRowStruct());
+	}
+}
+
 void FDataTableEditor::RefreshCachedDataTable(const FName InCachedSelection, const bool bUpdateEvenIfValid)
 {
 	const UDataTable* Table = GetDataTable();
diff --git a/Engine/Source/Editor/DataTableEditor/Private/DataTableEditor.h b/Engine/Source/Editor/DataTableEditor/Private/DataTableEditor.h
--- a/Engine/Source/Editor/DataTableEditor/Private/DataTableEditor.h
+++ b/Engine/Source/Editor/DataTableEditor/Private/DataTableEditor.h
@@ -89,6 +89,12 @@ protected:
 	
 	void OnFilterTextChanged(const FText& InFilterText);
 
+	virtual void PostRegenerateMenusAndToolbars() override;
+
+	FReply OnFindRowInContentBrowserClicked();
+
+	void OnNavigateToDataTableRowCode();
+
 	FText GetCellText(FDataTableEditorRowListViewDataPtr InRowDataPointer, int32 ColumnIndex) const;
 	FText GetCellToolTipText(FDataTableEditorRowListViewDataPtr InRowDataPointer, int32 ColumnIndex) const;
 
diff --git a/Engine/Source/Editor/UnrealEd/Private/SourceCodeNavigation.cpp b/Engine/Source/Editor/UnrealEd/Private/SourceCodeNavigation.cpp
--- a/Engine/Source/Editor/UnrealEd/Private/SourceCodeNavigation.cpp
+++ b/Engine/Source/Editor/UnrealEd/Private/SourceCodeNavigation.cpp
@@ -1599,6 +1599,52 @@ bool FSourceCodeNavigation::NavigateToProperty(const UProperty* InProperty)
 	return false;
 }
 
+bool FSourceCodeNavigation::CanNavigateToStruct(const UStruct* InStruct)
+{
+	if (!InStruct)
+	{
+		return false;
+	}
+
+	for (int32 i = 0; i < SourceCodeNavigationHandlers.Num(); ++i)
+	{
+		ISourceCodeNavigationHandler* handler = SourceCodeNavigationHandlers[i];
+		if (handler->CanNavigateToStruct(InStruct))
+		{
+			return true;
+		}
+	}
+
+	return InStruct->IsNative() && IsCompilerAvailable();
+}
+
+bool FSourceCodeNavigation::NavigateToStruct(const UStruct* InStruct)
+{
+	if (!InStruct)
+	{
+		return false;
+	}
+
+	for (int32 i = 0; i < SourceCodeNavigationHandlers.Num(); ++i)
+	{
+		ISourceCodeNavigationHandler* handler = SourceCodeNavigationHandlers[i];
+		if (handler->NavigateToStruct(InStruct))
+		{
+			return true;
+		}
+	}
+
+	FString ClassHeaderPath;
+	if (FSourceCodeNavigation::FindClassHeaderPath(InStruct, ClassHeaderPath) && IFileManager::Get().FileSize(*ClassHeaderPath) != INDEX_NONE)
+	{
+		FString AbsoluteHeaderPath = IFileManager::Get().ConvertToAbsolutePathForExternalAppForRead(*ClassHeaderPath);
+		FSourceCodeNavigation::OpenSourceFile(AbsoluteHeaderPath);
+		return true;
+	}
+
+	return false;
+}
+
 bool FSourceCodeNavigation::FindClassModuleName( UClass* InClass, FString& ModuleName )
 {
 	bool bResult = false;
diff --git a/Engine/Source/Editor/UnrealEd/Public/SourceCodeNavigation.h b/Engine/Source/Editor/UnrealEd/Public/SourceCodeNavigation.h
--- a/Engine/Source/Editor/UnrealEd/Public/SourceCodeNavigation.h
+++ b/Engine/Source/Editor/UnrealEd/Public/SourceCodeNavigation.h
@@ -76,7 +76,7 @@ public:
 	* Determines whether it is possible to navigate to the UFunction using this handler.
 	*
 	* @param	InFunction		Function to inspect
-	* @return					True if the class can be handled by this handler.
+	* @return					True if the function can be handled by this handler.
 	*/
 	virtual bool CanNavigateToFunction(const UFunction* InFunction) { return false; };
 
@@ -92,7 +92,7 @@ public:
 	* Determines whether it is possible to navigate to the UProperty using this handler.
 	*
 	* @param	InProperty		Property to inspect
-	* @return					True if the class can be handled by this handler.
+	* @return					True if the property can be handled by this handler.
 	*/
 	virtual bool CanNavigateToProperty(const UProperty* InProperty) { return false; };
 
@@ -103,6 +103,22 @@ public:
 	* @return					True if the property can be handled by this handler.
 	*/
 	virtual bool NavigateToProperty(const UProperty* InProperty) { return false; };
+
+	/**
+	* Determines whether it is possible to navigate to the UStruct using this handler.
+	*
+	* @param	InStruct		Struct to inspect
+	* @return					True if the struct can be handled by this handler.
+	*/
+	virtual bool CanNavigateToStruct(const UStruct* InStruct) { return false; };
+
+	/**
+	* Asynchronously navigates to a UStruct in an IDE or text editor.
+	*
+	* @param	InStruct		Struct in which to navigate.
+	* @return					True if the struct can be handled by this handler.
+	*/
+	virtual bool NavigateToStruct(const UStruct* InStruct) { return false; };
 };
 
 /**
@@ -241,6 +257,22 @@ public:
 	 */
 	UNREALED_API static bool NavigateToProperty(const UProperty* InProperty);
 
+	/**
+	 * Determines whether it is possible to navigate to the UStruct in the IDE
+	 *
+	 * @param	InStruct		UStruct to navigate to in source code
+	 * @return					Whether the navigation is likely to be successful or not
+ 	 */
+	UNREALED_API static bool CanNavigateToStruct(const UStruct* InStruct);
+
+	/**
+	 * Navigates asynchronously to the UStruct in the IDE
+	 *
+	 * @param	InStruct		UStruct to navigate to in source code
+	 * @return					Whether the navigation was successful or not
+	 */
+	UNREALED_API static bool NavigateToStruct(const UStruct* InStruct);
+
 	/** Delegate that's triggered when any symbol query has completed */
 	DECLARE_MULTICAST_DELEGATE( FOnSymbolQueryFinished );
 
