Add support for setting the ImportKeyField on a datatable, which allows selecting what field from the import/export data is used as the key
Add a DataTable Details tab to the data table editor to show this and the existing bStripFromClientBuilds flag
Right now you can only set these options after a datatable has been created, I will add import UI later as well as more options

diff --git a/Engine/Source/Editor/DataTableEditor/Private/DataTableEditor.cpp b/Engine/Source/Editor/DataTableEditor/Private/DataTableEditor.cpp
--- a/Engine/Source/Editor/DataTableEditor/Private/DataTableEditor.cpp
+++ b/Engine/Source/Editor/DataTableEditor/Private/DataTableEditor.cpp
@@ -22,10 +22,12 @@
 #include "Widgets/SToolTip.h"
 #include "Widgets/Text/STextBlock.h"
 #include "Widgets/Views/SListView.h"
+#include "PropertyEditorModule.h"
 
 #define LOCTEXT_NAMESPACE "DataTableEditor"
 
 const FName FDataTableEditor::DataTableTabId("DataTableEditor_DataTable");
+const FName FDataTableEditor::DataTableDetailsTabId("DataTableEditor_DataTableDetails");
 const FName FDataTableEditor::RowEditorTabId("DataTableEditor_RowEditor");
 const FName FDataTableEditor::RowNameColumnId("RowName");
 
@@ -74,6 +76,7 @@ void FDataTableEditor::RegisterTabSpawners(const TSharedRef<class FTabManager>&
 	FAssetEditorToolkit::RegisterTabSpawners(InTabManager);
 
 	CreateAndRegisterDataTableTab(InTabManager);
+	CreateAndRegisterDataTableDetailsTab(InTabManager);
 	CreateAndRegisterRowEditorTab(InTabManager);
 }
 
@@ -82,6 +85,7 @@ void FDataTableEditor::UnregisterTabSpawners(const TSharedRef<class FTabManager>
 	FAssetEditorToolkit::UnregisterTabSpawners(InTabManager);
 
 	InTabManager->UnregisterTabSpawner(DataTableTabId);
+	InTabManager->UnregisterTabSpawner(DataTableDetailsTabId);
 	InTabManager->UnregisterTabSpawner(RowEditorTabId);
 
 	DataTableTabWidget.Reset();
@@ -97,6 +101,17 @@ void FDataTableEditor::CreateAndRegisterDataTableTab(const TSharedRef<class FTab
 		.SetGroup(WorkspaceMenuCategory.ToSharedRef());
 }
 
+void FDataTableEditor::CreateAndRegisterDataTableDetailsTab(const TSharedRef<class FTabManager>& InTabManager)
+{
+	FPropertyEditorModule & EditModule = FModuleManager::Get().GetModuleChecked<FPropertyEditorModule>("PropertyEditor");
+	FDetailsViewArgs DetailsViewArgs(/*bUpdateFromSelection=*/ false, /*bLockable=*/ false, /*bAllowSearch=*/ true, /*InNameAreaSettings=*/ FDetailsViewArgs::HideNameArea, /*bHideSelectionTip=*/ true);
+	PropertyView = EditModule.CreateDetailView(DetailsViewArgs);
+
+	InTabManager->RegisterTabSpawner(DataTableDetailsTabId, FOnSpawnTab::CreateSP(this, &FDataTableEditor::SpawnTab_DataTableDetails))
+		.SetDisplayName(LOCTEXT("DataTableDetailsTab", "Data Table Details"))
+		.SetGroup(WorkspaceMenuCategory.ToSharedRef());
+}
+
 void FDataTableEditor::CreateAndRegisterRowEditorTab(const TSharedRef<class FTabManager>& InTabManager)
 {
 	RowEditorTabWidget = CreateRowEditorBox();
@@ -198,7 +213,7 @@ void FDataTableEditor::HandlePostChange()
 
 void FDataTableEditor::InitDataTableEditor( const EToolkitMode::Type Mode, const TSharedPtr< class IToolkitHost >& InitToolkitHost, UDataTable* Table )
 {
-	TSharedRef<FTabManager::FLayout> StandaloneDefaultLayout = FTabManager::NewLayout( "Standalone_DataTableEditor_Layout_v2" )
+	TSharedRef<FTabManager::FLayout> StandaloneDefaultLayout = FTabManager::NewLayout( "Standalone_DataTableEditor_Layout_v3" )
 	->AddArea
 	(
 		FTabManager::NewPrimaryArea()->SetOrientation(Orient_Vertical)
@@ -213,6 +228,8 @@ void FDataTableEditor::InitDataTableEditor( const EToolkitMode::Type Mode, const
 		(
 			FTabManager::NewStack()
 			->AddTab(DataTableTabId, ETabState::OpenedTab)
+			->AddTab(DataTableDetailsTabId, ETabState::OpenedTab)
+			->SetForegroundTab(DataTableTabId)
 		)
 		->Split
 		(
@@ -786,6 +803,26 @@ TSharedRef<SDockTab> FDataTableEditor::SpawnTab_DataTable( const FSpawnTabArgs&
 		];
 }
 
+TSharedRef<SDockTab> FDataTableEditor::SpawnTab_DataTableDetails(const FSpawnTabArgs& Args)
+{
+	check(Args.GetTabId().TabType == DataTableDetailsTabId);
+
+	PropertyView->SetObject(const_cast<UDataTable*>(GetDataTable()));
+
+	return SNew(SDockTab)
+		.Icon(FEditorStyle::GetBrush("DataTableEditor.Tabs.Properties"))
+		.Label(LOCTEXT("DataTableDetails", "Data Table Details"))
+		.TabColorScale(GetTabColorScale())
+		[
+			SNew(SBorder)
+			.Padding(2)
+			.BorderImage(FEditorStyle::GetBrush("ToolPanel.GroupBorder"))
+			[
+				PropertyView.ToSharedRef()
+			]
+		];
+}
+
 void FDataTableEditor::SetHighlightedRow(FName Name)
 {
 	if (Name == HighlightedRowName)
diff --git a/Engine/Source/Editor/DataTableEditor/Private/DataTableEditor.h b/Engine/Source/Editor/DataTableEditor/Private/DataTableEditor.h
--- a/Engine/Source/Editor/DataTableEditor/Private/DataTableEditor.h
+++ b/Engine/Source/Editor/DataTableEditor/Private/DataTableEditor.h
@@ -100,6 +100,9 @@ protected:
 	/**	Spawns the tab with the data table inside */
 	TSharedRef<SDockTab> SpawnTab_DataTable( const FSpawnTabArgs& Args );
 
+	/**	Spawns the tab with the data table inside */
+	TSharedRef<SDockTab> SpawnTab_DataTableDetails(const FSpawnTabArgs& Args);
+
 	/**	Spawns the tab with the Row Editor inside */
 	TSharedRef<SDockTab> SpawnTab_RowEditor(const FSpawnTabArgs& Args);
 
@@ -131,6 +134,9 @@ protected:
 	/** Helper function for creating and registering the tab containing the data table data */
 	virtual void CreateAndRegisterDataTableTab(const TSharedRef<class FTabManager>& InTabManager);
 
+	/** Helper function for creating and registering the tab containing the data table details */
+	virtual void CreateAndRegisterDataTableDetailsTab(const TSharedRef<class FTabManager>& InTabManager);
+
 	/** Helper function for creating and registering the tab containing the row editor */
 	virtual void CreateAndRegisterRowEditorTab(const TSharedRef<class FTabManager>& InTabManager);
 
@@ -154,6 +160,9 @@ protected:
 	/** UI for the "Data Table" tab */
 	TSharedPtr<SWidget> DataTableTabWidget;
 
+	/** Property viewing widget */
+	TSharedPtr<class IDetailsView> PropertyView;
+
 	/** UI for the "Row Editor" tab */
 	TSharedPtr<SWidget> RowEditorTabWidget;
 
@@ -197,6 +206,9 @@ protected:
 	/**	The tab id for the data table tab */
 	static const FName DataTableTabId;
 
+	/**	The tab id for the data table details tab */
+	static const FName DataTableDetailsTabId;
+
 	/**	The tab id for the row editor tab */
 	static const FName RowEditorTabId;
 
diff --git a/Engine/Source/Editor/UnrealEd/Private/Factories/CSVImportFactory.cpp b/Engine/Source/Editor/UnrealEd/Private/Factories/CSVImportFactory.cpp
--- a/Engine/Source/Editor/UnrealEd/Private/Factories/CSVImportFactory.cpp
+++ b/Engine/Source/Editor/UnrealEd/Private/Factories/CSVImportFactory.cpp
@@ -212,6 +212,7 @@ UObject* UCSVImportFactory::FactoryCreateText(UClass* InClass, UObject* InParent
 			// If there is an existing table, need to call this to free data memory before recreating object
 			bool bStripFromClientBuilds = false;
 			UDataTable::FOnDataTableChanged OldOnDataTableChanged;
+			FString ImportKeyField;
 			if (ExistingTable != nullptr)
 			{
 				OldOnDataTableChanged = MoveTemp(ExistingTable->OnDataTableChanged());
@@ -219,6 +220,7 @@ UObject* UCSVImportFactory::FactoryCreateText(UClass* InClass, UObject* InParent
 				DataTableClass = ExistingTable->GetClass();
 				ExistingTable->EmptyTable();
 				bStripFromClientBuilds = ExistingTable->bStripFromClientBuilds;
+				ImportKeyField = ExistingTable->ImportKeyField;
 			}
 
 			// Create/reset table
@@ -226,6 +228,7 @@ UObject* UCSVImportFactory::FactoryCreateText(UClass* InClass, UObject* InParent
 			NewTable->RowStruct = ImportRowStruct;
 			NewTable->AssetImportData->Update(CurrentFilename);
 			NewTable->bStripFromClientBuilds = bStripFromClientBuilds;
+			NewTable->ImportKeyField = ImportKeyField;
 			// Go ahead and create table from string
 			Problems = DoImportDataTable(NewTable, String);
 
diff --git a/Engine/Source/Runtime/Engine/Classes/Engine/DataTable.h b/Engine/Source/Runtime/Engine/Classes/Engine/DataTable.h
--- a/Engine/Source/Runtime/Engine/Classes/Engine/DataTable.h
+++ b/Engine/Source/Runtime/Engine/Classes/Engine/DataTable.h
@@ -52,7 +52,7 @@ struct FTableRowBase
 /**
  * Imported spreadsheet table.
  */
-UCLASS(MinimalAPI, BlueprintType)
+UCLASS(MinimalAPI, BlueprintType, AutoExpandCategories = DataTable)
 class UDataTable
 	: public UObject
 {
@@ -90,6 +90,10 @@ public:
 	UPROPERTY(EditAnywhere, Category=DataTable)
 	uint8 bStripFromClientBuilds : 1;
 
+	/** Explicit field in import data to use as key. If this is empty it uses Name for JSON and the first field found for CSV */
+	UPROPERTY(EditAnywhere, Category=DataTable)
+	FString ImportKeyField;
+	
 #if WITH_EDITOR
 	ENGINE_API virtual void PostEditChangeProperty(FPropertyChangedEvent& PropertyChangedEvent) override;
 #endif // WITH_EDITOR
@@ -317,7 +321,8 @@ public:
 	*/
 	ENGINE_API TArray<FString> CreateTableFromJSONString(const FString& InString);
 
-	TArray<UProperty*> GetTablePropertyArray(const TArray<const TCHAR*>& Cells, UStruct* RowStruct, TArray<FString>& OutProblems);
+	/** Get array of UProperties that corresponds to columns in the table */
+	TArray<UProperty*> GetTablePropertyArray(const TArray<const TCHAR*>& Cells, UStruct* RowStruct, TArray<FString>& OutProblems, bool bIncludeNameColumn = false);
 	
 	/** 
 	 *	Create table from another Data Table
diff --git a/Engine/Source/Runtime/Engine/Private/DataTable.cpp b/Engine/Source/Runtime/Engine/Private/DataTable.cpp
--- a/Engine/Source/Runtime/Engine/Private/DataTable.cpp
+++ b/Engine/Source/Runtime/Engine/Private/DataTable.cpp
@@ -558,22 +558,20 @@ bool UDataTable::WriteTableAsJSONObject(const TSharedRef< TJsonWriter<TCHAR, TPr
 }
 #endif
 
-/** Get array of UProperties that corresponds to columns in the table */
-TArray<UProperty*> UDataTable::GetTablePropertyArray(const TArray<const TCHAR*>& Cells, UStruct* InRowStruct, TArray<FString>& OutProblems)
+TArray<UProperty*> UDataTable::GetTablePropertyArray(const TArray<const TCHAR*>& Cells, UStruct* InRowStruct, TArray<FString>& OutProblems, bool bIncludeNameColumn)
 {
 	TArray<UProperty*> ColumnProps;
 
 	// Get list of all expected properties from the struct
 	TArray<FName> ExpectedPropNames = DataTableUtils::GetStructPropertyNames(InRowStruct);
 
-	// Need at least 2 columns, first column is skipped, will contain row names
+	// Need at least 2 columns, first column will contain row names
 	if(Cells.Num() > 1)
 	{
 		ColumnProps.AddZeroed( Cells.Num() );
 
-		// first element always NULL - as first column is row names
-
-		for (int32 ColIdx = 1; ColIdx < Cells.Num(); ++ColIdx)
+		// Skip first column depending on option
+		for (int32 ColIdx = bIncludeNameColumn ? 0 : 1; ColIdx < Cells.Num(); ++ColIdx)
 		{
 			const TCHAR* ColumnValue = Cells[ColIdx];
 
diff --git a/Engine/Source/Runtime/Engine/Private/DataTableCSV.cpp b/Engine/Source/Runtime/Engine/Private/DataTableCSV.cpp
--- a/Engine/Source/Runtime/Engine/Private/DataTableCSV.cpp
+++ b/Engine/Source/Runtime/Engine/Private/DataTableCSV.cpp
@@ -26,14 +26,35 @@ bool FDataTableExporterCSV::WriteTable(const UDataTable& InDataTable)
 	}
 
 	// Write the header (column titles)
-	ExportedText += TEXT("---");
+	FString ImportKeyField;
+	if (!InDataTable.ImportKeyField.IsEmpty())
+	{
+		// Write actual name if we have it
+		ImportKeyField = InDataTable.ImportKeyField;
+		ExportedText += ImportKeyField;
+	}
+	else
+	{
+		ExportedText += TEXT("---");
+	}
+
+	UProperty* SkipProperty = nullptr;
 	for (TFieldIterator<UProperty> It(InDataTable.RowStruct); It; ++It)
 	{
 		UProperty* BaseProp = *It;
 		check(BaseProp);
 
+		FString ColumnHeader = DataTableUtils::GetPropertyExportName(BaseProp, DTExportFlags);
+		
+		if (ColumnHeader == ImportKeyField)
+		{
+			// Don't write header again if this is the name field, and save for skipping later
+			SkipProperty = BaseProp;
+			continue;
+		}
+		
 		ExportedText += TEXT(",");
-		ExportedText += DataTableUtils::GetPropertyExportName(BaseProp, DTExportFlags);
+		ExportedText += ColumnHeader;
 	}
 	ExportedText += TEXT("\n");
 
@@ -52,7 +73,7 @@ bool FDataTableExporterCSV::WriteTable(const UDataTable& InDataTable)
 	return true;
 }
 
-bool FDataTableExporterCSV::WriteRow(const UScriptStruct* InRowStruct, const void* InRowData)
+bool FDataTableExporterCSV::WriteRow(const UScriptStruct* InRowStruct, const void* InRowData, const UProperty* SkipProperty)
 {
 	if (!InRowStruct)
 	{
@@ -64,6 +85,11 @@ bool FDataTableExporterCSV::WriteRow(const UScriptStruct* InRowStruct, const voi
 		UProperty* BaseProp = *It;
 		check(BaseProp);
 
+		if (BaseProp == SkipProperty)
+		{
+			continue;
+		}
+
 		const void* Data = BaseProp->ContainerPtrToValuePtr<void>(InRowData, 0);
 		WriteStructEntry(InRowData, BaseProp, Data);
 	}
@@ -124,7 +150,23 @@ bool FDataTableImporterCSV::ReadTable()
 	}
 
 	// Find property for each column
-	TArray<UProperty*> ColumnProps = DataTable->GetTablePropertyArray(Rows[0], DataTable->RowStruct, ImportProblems);
+	int32 KeyColumn = 0;
+	bool bIncludeNameColumn = false;
+	if (!DataTable->ImportKeyField.IsEmpty())
+	{
+		// Search for key column by name
+		for (int32 ColIdx = 0; ColIdx < Rows[0].Num(); ++ColIdx)
+		{
+			const TCHAR* ColumnValue = Rows[0][ColIdx];
+			if (DataTable->ImportKeyField == FString(ColumnValue))
+			{
+				KeyColumn = ColIdx;
+				break;
+			}
+		}
+	}
+	
+	TArray<UProperty*> ColumnProps = DataTable->GetTablePropertyArray(Rows[0], DataTable->RowStruct, ImportProblems, bIncludeNameColumn);
 
 	// Empty existing data
 	DataTable->EmptyTable();
@@ -134,8 +176,8 @@ bool FDataTableImporterCSV::ReadTable()
 	{
 		const TArray<const TCHAR*>& Cells = Rows[RowIdx];
 
-		// Need at least 1 cells (row name)
-		if(Cells.Num() < 1)
+		// Need at least the key column
+		if (Cells.Num() <= KeyColumn)
 		{
 			ImportProblems.Add(FString::Printf(TEXT("Row '%d' has too few cells."), RowIdx));
 			continue;
@@ -149,7 +191,7 @@ bool FDataTableImporterCSV::ReadTable()
 		}
 
 		// Get row name
-		FName RowName = DataTableUtils::MakeValidName(Cells[0]);
+		FName RowName = DataTableUtils::MakeValidName(Cells[KeyColumn]);
 
 		// Check its not 'none'
 		if(RowName == NAME_None)
@@ -173,8 +215,8 @@ bool FDataTableImporterCSV::ReadTable()
 		// Add to row map
 		DataTable->AddRowInternal(RowName, RowData);
 
-		// Now iterate over cells (skipping first cell, that was row name)
-		for(int32 CellIdx=1; CellIdx<Cells.Num(); CellIdx++)
+		// Now iterate over cells (skipping first cell unless we had an explicit name)
+		for(int32 CellIdx = bIncludeNameColumn ? 0 : 1; CellIdx<Cells.Num(); CellIdx++)
 		{
 			// Try and assign string to data using the column property
 			UProperty* ColumnProp = ColumnProps[CellIdx];
diff --git a/Engine/Source/Runtime/Engine/Private/DataTableCSV.h b/Engine/Source/Runtime/Engine/Private/DataTableCSV.h
--- a/Engine/Source/Runtime/Engine/Private/DataTableCSV.h
+++ b/Engine/Source/Runtime/Engine/Private/DataTableCSV.h
@@ -18,7 +18,7 @@ public:
 
 	bool WriteTable(const UDataTable& InDataTable);
 
-	bool WriteRow(const UScriptStruct* InRowStruct, const void* InRowData);
+	bool WriteRow(const UScriptStruct* InRowStruct, const void* InRowData, const UProperty* SkipProperty = nullptr);
 
 private:
 	bool WriteStructEntry(const void* InRowData, UProperty* InProperty, const void* InPropertyData);
diff --git a/Engine/Source/Runtime/Engine/Private/DataTableJSON.cpp b/Engine/Source/Runtime/Engine/Private/DataTableJSON.cpp
--- a/Engine/Source/Runtime/Engine/Private/DataTableJSON.cpp
+++ b/Engine/Source/Runtime/Engine/Private/DataTableJSON.cpp
@@ -97,6 +97,7 @@ bool FDataTableExporterJSON::WriteTable(const UDataTable& InDataTable)
 		return false;
 	}
 
+	FString KeyField = GetKeyFieldName(InDataTable);
 	JsonWriter->WriteArrayStart();
 
 	// Iterate over rows
@@ -106,11 +107,11 @@ bool FDataTableExporterJSON::WriteTable(const UDataTable& InDataTable)
 		{
 			// RowName
 			const FName RowName = RowIt.Key();
-			JsonWriter->WriteValue(TEXT("Name"), RowName.ToString());
+			JsonWriter->WriteValue(KeyField, RowName.ToString());
 
 			// Now the values
 			uint8* RowData = RowIt.Value();
-			WriteRow(InDataTable.RowStruct, RowData);
+			WriteRow(InDataTable.RowStruct, RowData, &KeyField);
 		}
 		JsonWriter->WriteObjectEnd();
 	}
@@ -147,23 +148,30 @@ bool FDataTableExporterJSON::WriteTableAsObject(const UDataTable& InDataTable)
 	return true;
 }
 
-bool FDataTableExporterJSON::WriteRow(const UScriptStruct* InRowStruct, const void* InRowData)
+bool FDataTableExporterJSON::WriteRow(const UScriptStruct* InRowStruct, const void* InRowData, const FString* FieldToSkip)
 {
 	if (!InRowStruct)
 	{
 		return false;
 	}
 
-	return WriteStruct(InRowStruct, InRowData);
+	return WriteStruct(InRowStruct, InRowData, FieldToSkip);
 }
 
-bool FDataTableExporterJSON::WriteStruct(const UScriptStruct* InStruct, const void* InStructData)
+bool FDataTableExporterJSON::WriteStruct(const UScriptStruct* InStruct, const void* InStructData, const FString* FieldToSkip)
 {
 	for (TFieldIterator<const UProperty> It(InStruct); It; ++It)
 	{
 		const UProperty* BaseProp = *It;
 		check(BaseProp);
 
+		const FString Identifier = DataTableUtils::GetPropertyExportName(BaseProp, DTExportFlags);
+		if (FieldToSkip && *FieldToSkip == Identifier)
+		{
+			// Skip this field
+			continue;
+		}
+ 
 		if (BaseProp->ArrayDim == 1)
 		{
 			const void* Data = BaseProp->ContainerPtrToValuePtr<void>(InStructData, 0);
@@ -171,8 +179,6 @@ bool FDataTableExporterJSON::WriteStruct(const UScriptStruct* InStruct, const vo
 		}
 		else
 		{
-			const FString Identifier = DataTableUtils::GetPropertyExportName(BaseProp, DTExportFlags);
-
 			JsonWriter->WriteArrayStart(Identifier);
 
 			for (int32 ArrayEntryIndex = 0; ArrayEntryIndex < BaseProp->ArrayDim; ++ArrayEntryIndex)
@@ -188,6 +194,19 @@ bool FDataTableExporterJSON::WriteStruct(const UScriptStruct* InStruct, const vo
 	return true;
 }
 
+FString FDataTableExporterJSON::GetKeyFieldName(const UDataTable& InDataTable)
+{
+	FString ExplicitString = InDataTable.ImportKeyField;
+	if (ExplicitString.IsEmpty())
+	{
+		return TEXT("Name");
+	}
+	else
+	{
+		return ExplicitString;
+	}
+}
+
 bool FDataTableExporterJSON::WriteStructEntry(const void* InRowData, const UProperty* InProperty, const void* InPropertyData)
 {
 	const FString Identifier = DataTableUtils::GetPropertyExportName(InProperty, DTExportFlags);
@@ -424,7 +443,8 @@ bool FDataTableImporterJSON::ReadTable()
 bool FDataTableImporterJSON::ReadRow(const TSharedRef<FJsonObject>& InParsedTableRowObject, const int32 InRowIdx)
 {
 	// Get row name
-	FName RowName = DataTableUtils::MakeValidName(InParsedTableRowObject->GetStringField(TEXT("Name")));
+	FString RowKey = FDataTableExporterJSON::GetKeyFieldName(*DataTable);
+	FName RowName = DataTableUtils::MakeValidName(InParsedTableRowObject->GetStringField(RowKey));
 
 	// Check its not 'none'
 	if (RowName.IsNone())
diff --git a/Engine/Source/Runtime/Engine/Private/DataTableJSON.h b/Engine/Source/Runtime/Engine/Private/DataTableJSON.h
--- a/Engine/Source/Runtime/Engine/Private/DataTableJSON.h
+++ b/Engine/Source/Runtime/Engine/Private/DataTableJSON.h
@@ -31,14 +31,20 @@ public:
 
 	~FDataTableExporterJSON();
 
+	/** Writes the data table out as an array of objects */
 	bool WriteTable(const UDataTable& InDataTable);
 
 	/** Writes the data table out as a named object with each row being a sub value on that object */
 	bool WriteTableAsObject(const UDataTable& InDataTable);
 
-	bool WriteRow(const UScriptStruct* InRowStruct, const void* InRowData);
+	/** Writes out a single row */
+	bool WriteRow(const UScriptStruct* InRowStruct, const void* InRowData, const FString* FieldToSkip = nullptr);
 
-	bool WriteStruct(const UScriptStruct* InStruct, const void* InStructData);
+	/** Writes the contents of a single row */
+	bool WriteStruct(const UScriptStruct* InStruct, const void* InStructData, const FString* FieldToSkip = nullptr);
+
+	/** Returns what string is used as the key/name field for a data table */
+	static FString GetKeyFieldName(const UDataTable& InDataTable);
 
 private:
 	bool WriteStructEntry(const void* InRowData, const UProperty* InProperty, const void* InPropertyData);
