[NO NON-PII COMMIT MESSAGE]

diff --git a/Engine/Source/Editor/DataTableEditor/Private/DataTableEditor.cpp b/Engine/Source/Editor/DataTableEditor/Private/DataTableEditor.cpp
--- a/Engine/Source/Editor/DataTableEditor/Private/DataTableEditor.cpp
+++ b/Engine/Source/Editor/DataTableEditor/Private/DataTableEditor.cpp
@@ -649,6 +649,11 @@ void FDataTableEditor::RefreshCachedDataTable(const FName InCachedSelection, con
 	}
 
 	UpdateVisibleRows(InCachedSelection, bUpdateEvenIfValid);
+
+	if (PropertyView.IsValid())
+	{
+		PropertyView->SetObject(const_cast<UDataTable*>(Table));
+	}
 }
 
 void FDataTableEditor::UpdateVisibleRows(const FName InCachedSelection, const bool bUpdateEvenIfValid)
