Fixed a typo in FDataTableEditor::OnNavigateToDataTableRowCode where it was trying to navigate to the struct twice

diff --git a/Engine/Source/Editor/DataTableEditor/Private/DataTableEditor.cpp b/Engine/Source/Editor/DataTableEditor/Private/DataTableEditor.cpp
--- a/Engine/Source/Editor/DataTableEditor/Private/DataTableEditor.cpp
+++ b/Engine/Source/Editor/DataTableEditor/Private/DataTableEditor.cpp
@@ -637,7 +637,7 @@ FReply FDataTableEditor::OnFindRowInContentBrowserClicked()
 void FDataTableEditor::OnNavigateToDataTableRowCode()
 {
 	const UDataTable* DataTable = GetDataTable();
-	if (DataTable && FSourceCodeNavigation::NavigateToStruct(DataTable->GetRowStruct()))
+	if (DataTable && FSourceCodeNavigation::CanNavigateToStruct(DataTable->GetRowStruct()))
 	{
 		FSourceCodeNavigation::NavigateToStruct(DataTable->GetRowStruct());
 	}
