Fix DataTableEditor member vars not being initialized

diff --git a/Engine/Source/Editor/DataTableEditor/Private/DataTableEditor.cpp b/Engine/Source/Editor/DataTableEditor/Private/DataTableEditor.cpp
--- a/Engine/Source/Editor/DataTableEditor/Private/DataTableEditor.cpp
+++ b/Engine/Source/Editor/DataTableEditor/Private/DataTableEditor.cpp
@@ -130,6 +130,10 @@ void FDataTableEditor::CreateAndRegisterRowEditorTab(const TSharedRef<class FTab
 }
 
 FDataTableEditor::FDataTableEditor()
+	: RowNameColumnWidth(0)
+	, RowNumberColumnWidth(0)
+	, HighlightedVisibleRowIndex(INDEX_NONE)
+	, SortMode(EColumnSortMode::Ascending)
 {
 }
 
@@ -1406,7 +1410,7 @@ void FDataTableEditor::SetHighlightedRow(FName Name)
 	{
 		HighlightedRowName = NAME_None;
 		CellsListView->ClearSelection();
-		HighlightedVisibleRowIndex = -1;
+		HighlightedVisibleRowIndex = INDEX_NONE;
 	}
 	else
 	{
