Make Data Table Editor refresh when filter changes

diff --git a/Engine/Source/Editor/DataTableEditor/Private/DataTableEditor.cpp b/Engine/Source/Editor/DataTableEditor/Private/DataTableEditor.cpp
--- a/Engine/Source/Editor/DataTableEditor/Private/DataTableEditor.cpp
+++ b/Engine/Source/Editor/DataTableEditor/Private/DataTableEditor.cpp
@@ -925,33 +925,16 @@ FText FDataTableEditor::GetFilterText() const
 
 void FDataTableEditor::OnFilterTextChanged(const FText& InFilterText)
 {
-	if (InFilterText.IsEmpty())
-	{
-		OnFilterCleared();
-	}
-	else
-	{
-		ActiveFilterText = InFilterText;
-		UpdateVisibleRows();
-	}
+	ActiveFilterText = InFilterText;
+	UpdateVisibleRows();
 }
 
 void FDataTableEditor::OnFilterTextCommitted(const FText& NewText, ETextCommit::Type CommitInfo)
 {
 	if (CommitInfo == ETextCommit::OnCleared)
 	{
-		OnFilterCleared();
-	}
-}
-
-void FDataTableEditor::OnFilterCleared()
-{
-	ActiveFilterText = FText();
-	if (VisibleRows.IsValidIndex(HighlightedVisibleRowIndex))
-	{
-		VisibleRows = AvailableRows;
-		SearchBoxWidget->SetText(ActiveFilterText);
-		CellsListView->RequestListRefresh();
+		SearchBoxWidget->SetText(FText::GetEmpty());
+		OnFilterTextChanged(FText::GetEmpty());
 	}
 }
 
@@ -1215,6 +1198,9 @@ void FDataTableEditor::UpdateVisibleRows(const FName InCachedSelection, const bo
 			}
 		}
 	}
+
+	CellsListView->RequestListRefresh();
+	RestoreCachedSelection(InCachedSelection, bUpdateEvenIfValid);
 }
 
 void FDataTableEditor::RestoreCachedSelection(const FName InCachedSelection, const bool bUpdateEvenIfValid)
diff --git a/Engine/Source/Editor/DataTableEditor/Private/DataTableEditor.h b/Engine/Source/Editor/DataTableEditor/Private/DataTableEditor.h
--- a/Engine/Source/Editor/DataTableEditor/Private/DataTableEditor.h
+++ b/Engine/Source/Editor/DataTableEditor/Private/DataTableEditor.h
@@ -91,11 +91,8 @@ protected:
 	void RestoreCachedSelection(const FName InCachedSelection, const bool bUpdateEvenIfValid = false);
 	
 	void OnFilterTextChanged(const FText& InFilterText);
-
 	void OnFilterTextCommitted(const FText& NewText, ETextCommit::Type CommitInfo);
 
-	void OnFilterCleared();
-
 	virtual void PostRegenerateMenusAndToolbars() override;
 
 	FReply OnFindRowInContentBrowserClicked();
