PR #4208: Added support for copy, paste and duplicate on data table rows (Contributed by [REDACTED_PERSON])

diff --git a/Engine/Source/Editor/DataTableEditor/DataTableEditor.Build.cs b/Engine/Source/Editor/DataTableEditor/DataTableEditor.Build.cs
--- a/Engine/Source/Editor/DataTableEditor/DataTableEditor.Build.cs
+++ b/Engine/Source/Editor/DataTableEditor/DataTableEditor.Build.cs
@@ -12,6 +12,7 @@ public class DataTableEditor : ModuleRules
 			new string[] {
 				"Core",
 				"CoreUObject",
+				"ApplicationCore",
 				"Engine",
                 "InputCore",
 				"Slate",
diff --git a/Engine/Source/Editor/DataTableEditor/Private/DataTableEditor.cpp b/Engine/Source/Editor/DataTableEditor/Private/DataTableEditor.cpp
--- a/Engine/Source/Editor/DataTableEditor/Private/DataTableEditor.cpp
+++ b/Engine/Source/Editor/DataTableEditor/Private/DataTableEditor.cpp
@@ -27,6 +27,14 @@
 #include "Widgets/Input/SHyperlink.h"
 #include "SourceCodeNavigation.h"
 
+#include "Framework/Commands/GenericCommands.h"
+#include "Framework/MultiBox/MultiBoxBuilder.h"
+#include "UObject/StructOnScope.h"
+#include "HAL/PlatformApplicationMisc.h"
+#include "Misc/FeedbackContext.h"
+#include "Widgets/Notifications/SNotificationList.h"
+#include "Framework/Notifications/NotificationManager.h"
+
 #define LOCTEXT_NAMESPACE "DataTableEditor"
 
 const FName FDataTableEditor::DataTableTabId("DataTableEditor_DataTable");
@@ -245,7 +253,10 @@ void FDataTableEditor::InitDataTableEditor( const EToolkitMode::Type Mode, const
 		SpawnToolkitTab( DataTableTabId, TabInitializationPayload, EToolkitTabSpot::Details );
 	}*/
 
-	// NOTE: Could fill in asset editor commands here!
+	// asset editor commands here
+	ToolkitCommands->MapAction(FGenericCommands::Get().Copy, FExecuteAction::CreateSP(this, &FDataTableEditor::CopySelectedRow));
+	ToolkitCommands->MapAction(FGenericCommands::Get().Paste, FExecuteAction::CreateSP(this, &FDataTableEditor::PasteOnSelectedRow));
+	ToolkitCommands->MapAction(FGenericCommands::Get().Duplicate, FExecuteAction::CreateSP(this, &FDataTableEditor::DuplicateSelectedRow));
 }
 
 FName FDataTableEditor::GetToolkitFName() const
@@ -458,6 +469,65 @@ void FDataTableEditor::OnRowSelectionChanged(FDataTableEditorRowListViewDataPtr
 	}
 }
 
+void FDataTableEditor::CopySelectedRow()
+{
+	UDataTable* TablePtr = Cast<UDataTable>(GetEditingObject());
+	uint8* RowPtr = TablePtr ? TablePtr->GetRowMap().FindRef(HighlightedRowName) : nullptr;
+
+	if (!RowPtr || !TablePtr->RowStruct)
+		return;
+
+	FString ClipboardValue;
+	TablePtr->RowStruct->ExportText(ClipboardValue, RowPtr, RowPtr, TablePtr, PPF_Copy, nullptr);
+
+	FPlatformApplicationMisc::ClipboardCopy(*ClipboardValue);
+}
+
+void FDataTableEditor::PasteOnSelectedRow()
+{
+	UDataTable* TablePtr = Cast<UDataTable>(GetEditingObject());
+	uint8* RowPtr = TablePtr ? TablePtr->GetRowMap().FindRef(HighlightedRowName) : nullptr;
+
+	if (!RowPtr || !TablePtr->RowStruct)
+		return;
+
+	const FScopedTransaction Transaction(LOCTEXT("PasteDataTableRow", "Paste Data Table Row"));
+	TablePtr->Modify();
+
+	FString ClipboardValue;
+	FPlatformApplicationMisc::ClipboardPaste(ClipboardValue);
+
+	FDataTableEditorUtils::BroadcastPreChange(TablePtr, FDataTableEditorUtils::EDataTableChangeInfo::RowData);
+
+	const TCHAR* Result = TablePtr->RowStruct->ImportText(*ClipboardValue, RowPtr, TablePtr, PPF_Copy, GWarn, GetPathNameSafe(TablePtr->RowStruct));
+
+	FDataTableEditorUtils::BroadcastPostChange(TablePtr, FDataTableEditorUtils::EDataTableChangeInfo::RowData);
+
+	if (Result == nullptr)
+	{
+		FNotificationInfo Info(LOCTEXT("FailedPaste", "Failed to paste row"));
+		FSlateNotificationManager::Get().AddNotification(Info);
+	}
+}
+
+void FDataTableEditor::DuplicateSelectedRow()
+{
+	UDataTable* TablePtr = Cast<UDataTable>(GetEditingObject());
+	FName NewName = HighlightedRowName;
+
+	if (NewName == NAME_None || TablePtr == nullptr)
+		return;
+
+	const TArray<FName> ExistingNames = TablePtr->GetRowNames();
+	while (ExistingNames.Contains(NewName))
+	{
+		NewName.SetNumber(NewName.GetNumber() + 1);
+	}
+
+	FDataTableEditorUtils::DuplicateRow(TablePtr, HighlightedRowName, NewName);
+	FDataTableEditorUtils::SelectRow(TablePtr, NewName);
+}
+
 FText FDataTableEditor::GetFilterText() const
 {
 	return ActiveFilterText;
diff --git a/Engine/Source/Editor/DataTableEditor/Private/DataTableEditor.h b/Engine/Source/Editor/DataTableEditor/Private/DataTableEditor.h
--- a/Engine/Source/Editor/DataTableEditor/Private/DataTableEditor.h
+++ b/Engine/Source/Editor/DataTableEditor/Private/DataTableEditor.h
@@ -134,6 +134,10 @@ protected:
 
 	void OnRowSelectionChanged(FDataTableEditorRowListViewDataPtr InNewSelection, ESelectInfo::Type InSelectInfo);
 
+	void CopySelectedRow();
+	void PasteOnSelectedRow();
+	void DuplicateSelectedRow();
+
 	/** Helper function for creating and registering the tab containing the data table data */
 	virtual void CreateAndRegisterDataTableTab(const TSharedRef<class FTabManager>& InTabManager);
 
diff --git a/Engine/Source/Editor/DataTableEditor/Private/SDataTableListViewRowName.cpp b/Engine/Source/Editor/DataTableEditor/Private/SDataTableListViewRowName.cpp
--- a/Engine/Source/Editor/DataTableEditor/Private/SDataTableListViewRowName.cpp
+++ b/Engine/Source/Editor/DataTableEditor/Private/SDataTableListViewRowName.cpp
@@ -36,9 +36,11 @@ void SDataTableListViewRowName::Construct(const FArguments& InArgs, const TShare
 
 FReply SDataTableListViewRowName::OnMouseButtonUp(const FGeometry& MyGeometry, const FPointerEvent& MouseEvent)
 {
-	if (MouseEvent.GetEffectingButton() == EKeys::RightMouseButton && RowDataPtr.IsValid() && FEditorDelegates::OnOpenReferenceViewer.IsBound())
+	if (MouseEvent.GetEffectingButton() == EKeys::RightMouseButton && RowDataPtr.IsValid() && FEditorDelegates::OnOpenReferenceViewer.IsBound() && DataTableEditor.IsValid())
 	{
-		TSharedRef<SWidget> MenuWidget = FDataTableRowUtils::MakeRowActionsMenu(FExecuteAction::CreateSP(this, &SDataTableListViewRowName::OnSearchForReferences));
+		FDataTableEditorUtils::SelectRow(DataTableEditor.Pin()->GetDataTable(), RowDataPtr->RowId);
+
+		TSharedRef<SWidget> MenuWidget = FDataTableRowUtils::MakeRowActionsMenu(DataTableEditor.Pin(), FExecuteAction::CreateSP(this, &SDataTableListViewRowName::OnSearchForReferences));
 		
 		FWidgetPath WidgetPath = MouseEvent.GetEventPath() != nullptr ? *MouseEvent.GetEventPath() : FWidgetPath();
 
diff --git a/Engine/Source/Editor/DataTableEditor/Public/DataTableRowUtlis.cpp b/Engine/Source/Editor/DataTableEditor/Public/DataTableRowUtlis.cpp
--- a/Engine/Source/Editor/DataTableEditor/Public/DataTableRowUtlis.cpp
+++ b/Engine/Source/Editor/DataTableEditor/Public/DataTableRowUtlis.cpp
@@ -6,17 +6,23 @@
 #include "Framework/MultiBox/MultiBoxBuilder.h"
 #include "Textures/SlateIcon.h"
 #include "Widgets/SWidget.h"
+#include "Framework/Commands/GenericCommands.h"
+#include "IDataTableEditor.h"
 
 #define LOCTEXT_NAMESPACE "FDataTableRowUtils"
 
 const FText FDataTableRowUtils::SearchForReferencesActionName = LOCTEXT("FDataTableRowUtils_SearchForReferences", "Find Row References");
 const FText FDataTableRowUtils::SearchForReferencesActionTooltip = LOCTEXT("FDataTableRowUtils_SearchForReferencesTooltip", "Find assets that reference this Row");
 
-TSharedRef<SWidget> FDataTableRowUtils::MakeRowActionsMenu(FExecuteAction SearchForReferencesAction)
+TSharedRef<SWidget> FDataTableRowUtils::MakeRowActionsMenu(TSharedPtr<IDataTableEditor> Editor, FExecuteAction SearchForReferencesAction)
 {
 	if (SearchForReferencesAction.IsBound())
 	{
-		FMenuBuilder MenuBuilder(true, nullptr);
+		FMenuBuilder MenuBuilder(true, Editor->GetToolkitCommands());
+		MenuBuilder.AddMenuEntry(FGenericCommands::Get().Copy);
+		MenuBuilder.AddMenuEntry(FGenericCommands::Get().Paste);
+		MenuBuilder.AddMenuEntry(FGenericCommands::Get().Duplicate);
+		MenuBuilder.AddMenuSeparator();
 		MenuBuilder.AddMenuEntry(SearchForReferencesActionName, SearchForReferencesActionTooltip, 
 			FSlateIcon(), FUIAction(SearchForReferencesAction));
 		return MenuBuilder.MakeWidget();
diff --git a/Engine/Source/Editor/DataTableEditor/Public/DataTableRowUtlis.h b/Engine/Source/Editor/DataTableEditor/Public/DataTableRowUtlis.h
--- a/Engine/Source/Editor/DataTableEditor/Public/DataTableRowUtlis.h
+++ b/Engine/Source/Editor/DataTableEditor/Public/DataTableRowUtlis.h
@@ -12,7 +12,7 @@ class FDetailWidgetRow;
 class DATATABLEEDITOR_API FDataTableRowUtils
 {
 public:
-	static TSharedRef<SWidget> MakeRowActionsMenu(FExecuteAction SearchForReferencesAction);
+	static TSharedRef<SWidget> MakeRowActionsMenu(TSharedPtr<class IDataTableEditor> Editor, FExecuteAction SearchForReferencesAction);
 	static void AddSearchForReferencesContextMenu(FDetailWidgetRow& RowNameDetailWidget, FExecuteAction SearchForReferencesAction);
 
 private:
diff --git a/Engine/Source/Editor/UnrealEd/Private/DataTableEditorUtils.cpp b/Engine/Source/Editor/UnrealEd/Private/DataTableEditorUtils.cpp
--- a/Engine/Source/Editor/UnrealEd/Private/DataTableEditorUtils.cpp
+++ b/Engine/Source/Editor/UnrealEd/Private/DataTableEditorUtils.cpp
@@ -68,6 +68,31 @@ uint8* FDataTableEditorUtils::AddRow(UDataTable* DataTable, FName RowName)
 	return RowData;
 }
 
+uint8* FDataTableEditorUtils::DuplicateRow(UDataTable* DataTable, FName SourceRowName, FName RowName)
+{
+	if (!DataTable || (SourceRowName == NAME_None) || !DataTable->RowMap.Contains(SourceRowName) || DataTable->RowMap.Contains(RowName) || !DataTable->RowStruct)
+	{
+		return NULL;
+	}
+
+	const FScopedTransaction Transaction(LOCTEXT("DuplicateDataTableRow", "Duplicate Data Table Row"));
+
+	BroadcastPreChange(DataTable, EDataTableChangeInfo::RowList);
+	DataTable->Modify();
+
+	// Allocate data to store information, using UScriptStruct to know its size
+	uint8* OldRowData = *DataTable->RowMap.Find(SourceRowName);
+	uint8* NewRowData = (uint8*)FMemory::Malloc(DataTable->RowStruct->GetStructureSize());
+
+	DataTable->RowStruct->InitializeStruct(NewRowData);
+	DataTable->RowStruct->CopyScriptStruct(NewRowData, OldRowData);
+
+	// Add to row map
+	DataTable->RowMap.Add(RowName, NewRowData);
+	BroadcastPostChange(DataTable, EDataTableChangeInfo::RowList);
+	return NewRowData;
+}
+
 bool FDataTableEditorUtils::RenameRow(UDataTable* DataTable, FName OldName, FName NewName)
 {
 	bool bResult = false;
@@ -167,7 +192,7 @@ bool FDataTableEditorUtils::MoveRow(UDataTable* DataTable, FName RowName, ERowMo
 	return true;
 }
 
-bool FDataTableEditorUtils::SelectRow(UDataTable* DataTable, FName RowName)
+bool FDataTableEditorUtils::SelectRow(const UDataTable* DataTable, FName RowName)
 {
 	for (auto Listener : FDataTableEditorManager::Get().GetListeners())
 	{
diff --git a/Engine/Source/Editor/UnrealEd/Public/DataTableEditorUtils.h b/Engine/Source/Editor/UnrealEd/Public/DataTableEditorUtils.h
--- a/Engine/Source/Editor/UnrealEd/Public/DataTableEditorUtils.h
+++ b/Engine/Source/Editor/UnrealEd/Public/DataTableEditorUtils.h
@@ -71,9 +71,10 @@ struct UNREALED_API FDataTableEditorUtils
 
 	static bool RemoveRow(UDataTable* DataTable, FName Name);
 	static uint8* AddRow(UDataTable* DataTable, FName RowName);
+	static uint8* DuplicateRow(UDataTable* DataTable, FName SourceRowName, FName RowName);
 	static bool RenameRow(UDataTable* DataTable, FName OldName, FName NewName);
 	static bool MoveRow(UDataTable* DataTable, FName RowName, ERowMoveDirection Direction, int32 NumRowsToMoveBy = 1);
-	static bool SelectRow(UDataTable* DataTable, FName RowName);
+	static bool SelectRow(const UDataTable* DataTable, FName RowName);
 	static bool DiffersFromDefault(UDataTable* DataTable, FName RowName);
 	static bool ResetToDefault(UDataTable* DataTable, FName RowName);
 
