Add missing header in DataTableEditor.cpp

diff --git a/Engine/Source/Editor/DataTableEditor/Private/DataTableEditor.cpp b/Engine/Source/Editor/DataTableEditor/Private/DataTableEditor.cpp
--- a/Engine/Source/Editor/DataTableEditor/Private/DataTableEditor.cpp
+++ b/Engine/Source/Editor/DataTableEditor/Private/DataTableEditor.cpp
@@ -7,11 +7,17 @@
 #include "EditorStyleSet.h"
 #include "Fonts/FontMeasure.h"
 #include "Framework/Application/SlateApplication.h"
+#include "Framework/Commands/GenericCommands.h"
 #include "Framework/Layout/Overscroll.h"
+#include "Framework/MultiBox/MultiBoxBuilder.h"
+#include "Framework/Notifications/NotificationManager.h"
+#include "HAL/PlatformApplicationMisc.h"
 #include "IDocumentation.h"
+#include "Misc/FeedbackContext.h"
 #include "Misc/FileHelper.h"
 #include "Modules/ModuleManager.h"
 #include "Policies/PrettyJsonPrintPolicy.h"
+#include "ScopedTransaction.h"
 #include "SDataTableListViewRowName.h"
 #include "Serialization/JsonReader.h"
 #include "Serialization/JsonSerializer.h"
@@ -25,15 +31,12 @@
 #include "Widgets/Input/SButton.h"
 #include "Widgets/Images/SImage.h"
 #include "Widgets/Input/SHyperlink.h"
+#include "Widgets/Notifications/SNotificationList.h"
 #include "SourceCodeNavigation.h"
-
-#include "Framework/Commands/GenericCommands.h"
-#include "Framework/MultiBox/MultiBoxBuilder.h"
 #include "UObject/StructOnScope.h"
-#include "HAL/PlatformApplicationMisc.h"
-#include "Misc/FeedbackContext.h"
-#include "Widgets/Notifications/SNotificationList.h"
-#include "Framework/Notifications/NotificationManager.h"
+
+
+
 
 #define LOCTEXT_NAMESPACE "DataTableEditor"
 
