Updated the DataTable documentation link to take the user to the appropriate webpage. Added inline editing of the 'Name' row. Added Copy/Duplicate/Delete functionality to the whole row for context menu and not just the 'Name' row. Removed 'Row Name' section from Row Editor and shifted the arrangement tools.

diff --git a/Engine/Source/Editor/DataTableEditor/Private/DataTableEditor.cpp b/Engine/Source/Editor/DataTableEditor/Private/DataTableEditor.cpp
--- a/Engine/Source/Editor/DataTableEditor/Private/DataTableEditor.cpp
+++ b/Engine/Source/Editor/DataTableEditor/Private/DataTableEditor.cpp
@@ -35,7 +35,8 @@
 #include "SourceCodeNavigation.h"
 #include "PropertyEditorModule.h"
 #include "UObject/StructOnScope.h"
-
+#include "Toolkits/GlobalEditorCommonCommands.h"
+#include "DataTableRowUtlis.h"
 
 
 
@@ -68,6 +69,26 @@ public:
 			);
 	}
 
+
+	virtual FReply OnMouseButtonUp(const FGeometry& MyGeometry, const FPointerEvent& MouseEvent) override
+	{
+		if (MouseEvent.GetEffectingButton() == EKeys::RightMouseButton && Item.IsValid() && FEditorDelegates::OnOpenReferenceViewer.IsBound() && DataTableEditor.IsValid())
+		{
+			FDataTableEditorUtils::SelectRow(DataTableEditor.Pin()->GetDataTable(), Item->RowId);
+	
+			TSharedRef<SWidget> MenuWidget = FDataTableRowUtils::MakeRowActionsMenu(DataTableEditor.Pin(), FExecuteAction::CreateSP(this, &SDataTableListViewRow::OnSearchForReferences));
+	
+			FWidgetPath WidgetPath = MouseEvent.GetEventPath() != nullptr ? *MouseEvent.GetEventPath() : FWidgetPath();
+	
+			FSlateApplication::Get().PushMenu(AsShared(), WidgetPath, MenuWidget, MouseEvent.GetScreenSpacePosition(), FPopupTransitionEffect::ContextMenu);
+	
+			return FReply::Handled();
+		}
+	
+		return STableRow::OnMouseButtonUp(MyGeometry, MouseEvent);
+	}
+
+
 	/** Overridden from SMultiColumnTableRow.  Generates a widget for this column of the list view. */
 	virtual TSharedRef<SWidget> GenerateWidgetForColumn(const FName& ColumnName) override
 	{
@@ -78,6 +99,22 @@ public:
 	}
 
 private:
+	void OnSearchForReferences()
+	{
+		if (DataTableEditor.IsValid() && Item.IsValid())
+		{
+			if (FDataTableEditor* DataTableEditorPtr = DataTableEditor.Pin().Get())
+			{
+				UDataTable* SourceDataTable = const_cast<UDataTable*>(DataTableEditorPtr->GetDataTable());
+
+				TArray<FAssetIdentifier> AssetIdentifiers;
+				AssetIdentifiers.Add(FAssetIdentifier(SourceDataTable, Item->RowId));
+
+				FEditorDelegates::OnOpenReferenceViewer.Broadcast(AssetIdentifiers);
+			}
+		}
+	}
+
 	/** Weak reference to the data table editor that owns our list */
 	TWeakPtr<FDataTableEditor> DataTableEditor;
 	/** The item associated with this row of data */
@@ -273,6 +310,8 @@ void FDataTableEditor::InitDataTableEditor( const EToolkitMode::Type Mode, const
 		SpawnToolkitTab( DataTableTabId, TabInitializationPayload, EToolkitTabSpot::Details );
 	}*/
 
+	ToolkitCommands->MapAction(FGlobalEditorCommonCommands::Get().OpenDocumentation, FExecuteAction::CreateSP(this, &FDataTableEditor::BrowseDocumentation_Execute));
+
 	// asset editor commands here
 	ToolkitCommands->MapAction(FGenericCommands::Get().Copy, FExecuteAction::CreateSP(this, &FDataTableEditor::CopySelectedRow));
 	ToolkitCommands->MapAction(FGenericCommands::Get().Paste, FExecuteAction::CreateSP(this, &FDataTableEditor::PasteOnSelectedRow));
@@ -284,6 +323,16 @@ FName FDataTableEditor::GetToolkitFName() const
 	return FName("DataTableEditor");
 }
 
+void FDataTableEditor::BrowseDocumentation_Execute() const
+{
+	IDocumentation::Get()->Open(GetDocumentationLink(), FDocumentationSourceInfo(TEXT("help_menu_asset")));
+}
+
+FString FDataTableEditor::GetDocumentationLink() const
+{
+	return FString(TEXT("Gameplay/DataDriven"));
+}
+
 FText FDataTableEditor::GetBaseToolkitName() const
 {
 	return LOCTEXT( "AppLabel", "DataTable Editor" );
diff --git a/Engine/Source/Editor/DataTableEditor/Private/DataTableEditor.h b/Engine/Source/Editor/DataTableEditor/Private/DataTableEditor.h
--- a/Engine/Source/Editor/DataTableEditor/Private/DataTableEditor.h
+++ b/Engine/Source/Editor/DataTableEditor/Private/DataTableEditor.h
@@ -27,6 +27,7 @@ class FDataTableEditor : public IDataTableEditor
 	, public FDataTableEditorUtils::INotifyOnDataTableChanged
 {
 	friend class SDataTableListViewRow;
+	friend class SDataTableListViewRowName;
 
 public:
 
@@ -150,6 +151,10 @@ protected:
 	/** Helper function for creating and registering the tab containing the row editor */
 	virtual void CreateAndRegisterRowEditorTab(const TSharedRef<class FTabManager>& InTabManager);
 
+	void BrowseDocumentation_Execute() const;
+
+	virtual FString GetDocumentationLink() const override;
+
 protected:
 
 	/** Struct holding information about the current column widths */
diff --git a/Engine/Source/Editor/DataTableEditor/Private/SDataTableListViewRowName.cpp b/Engine/Source/Editor/DataTableEditor/Private/SDataTableListViewRowName.cpp
--- a/Engine/Source/Editor/DataTableEditor/Private/SDataTableListViewRowName.cpp
+++ b/Engine/Source/Editor/DataTableEditor/Private/SDataTableListViewRowName.cpp
@@ -8,13 +8,16 @@
 #include "EditorStyleSet.h"
 #include "Framework/MultiBox/MultiBoxBuilder.h"
 
+#define LOCTEXT_NAMESPACE "SDataTableListViewRowName"
+
 void SDataTableListViewRowName::Construct(const FArguments& InArgs, const TSharedRef<STableViewBase>& InOwnerTableView)
 {
 	RowDataPtr = InArgs._RowDataPtr;
+	CurrentName = MakeShareable(new FName(RowDataPtr->RowId));
 	DataTableEditor = InArgs._DataTableEditor;
 	STableRow<FDataTableEditorRowListViewDataPtr>::Construct(
 		typename STableRow<FDataTableEditorRowListViewDataPtr>::FArguments()
-		.Style(FEditorStyle::Get(), "DataTableEditor.NameListViewRow")
+		.Style(FEditorStyle::Get(), "DataTableEditor.CellListViewRow")
 		.Content()
 		[
 			SNew(SBox)
@@ -23,10 +26,10 @@ void SDataTableListViewRowName::Construct(const FArguments& InArgs, const TShare
 				SNew(SBox)
 				.HeightOverride(RowDataPtr->DesiredRowHeight)
 				[
-					SNew(STextBlock)
-					.ColorAndOpacity(DataTableEditor.Pin().Get(), &FDataTableEditor::GetRowTextColor, RowDataPtr->RowId)
+					SAssignNew(EditableText, SEditableText)
 					.Text(RowDataPtr->DisplayName)
-					.HighlightText(DataTableEditor.Pin().Get(), &FDataTableEditor::GetFilterText)
+					.OnTextCommitted(this, &SDataTableListViewRowName::OnRowRenamed)
+					.ColorAndOpacity(DataTableEditor.Pin().Get(), &FDataTableEditor::GetRowTextColor, RowDataPtr->RowId)					
 				]
 			]
 		],
@@ -39,13 +42,10 @@ FReply SDataTableListViewRowName::OnMouseButtonUp(const FGeometry& MyGeometry, c
 	if (MouseEvent.GetEffectingButton() == EKeys::RightMouseButton && RowDataPtr.IsValid() && FEditorDelegates::OnOpenReferenceViewer.IsBound() && DataTableEditor.IsValid())
 	{
 		FDataTableEditorUtils::SelectRow(DataTableEditor.Pin()->GetDataTable(), RowDataPtr->RowId);
-
 		TSharedRef<SWidget> MenuWidget = FDataTableRowUtils::MakeRowActionsMenu(DataTableEditor.Pin(), FExecuteAction::CreateSP(this, &SDataTableListViewRowName::OnSearchForReferences));
 		
 		FWidgetPath WidgetPath = MouseEvent.GetEventPath() != nullptr ? *MouseEvent.GetEventPath() : FWidgetPath();
-
 		FSlateApplication::Get().PushMenu(AsShared(), WidgetPath, MenuWidget, MouseEvent.GetScreenSpacePosition(), FPopupTransitionEffect::ContextMenu);
-
 		return FReply::Handled();
 	}
 
@@ -67,3 +67,92 @@ void SDataTableListViewRowName::OnSearchForReferences()
 		}
 	}
 }
+
+void SDataTableListViewRowName::Tick(const FGeometry& AllottedGeometry, const double InCurrentTime, const float InDeltaTime)
+{
+	SWidget::Tick(AllottedGeometry, InCurrentTime, InDeltaTime);
+	
+	if (EditableText->HasAnyUserFocus())
+	{
+		UDataTable* DataTable = Cast<UDataTable>(DataTableEditor.Pin()->GetEditingObject());
+		if (RowDataPtr.IsValid() && DataTable)
+		{
+	
+			DataTableEditor.Pin()->SelectionChange(DataTable, RowDataPtr->RowId);
+		}
+	}
+}
+
+FReply SDataTableListViewRowName::OnKeyDown(const FGeometry& MyGeometry, const FKeyEvent& InKeyEvent)
+{
+	FKey Key = InKeyEvent.GetKey();
+
+	if (Key == EKeys::Escape && EditableText->HasKeyboardFocus())
+	{
+		// Clear focus
+		return FReply::Handled().SetUserFocus(SharedThis(this), EFocusCause::Cleared);
+	}
+
+	return FReply::Unhandled();
+}
+
+void SDataTableListViewRowName::OnRowRenamed(const FText& Text, ETextCommit::Type CommitType)
+{
+	UDataTable* DataTable = Cast<UDataTable>(DataTableEditor.Pin()->GetEditingObject());
+
+	if (!GetCurrentNameAsText().EqualTo(Text) && DataTable)
+	{
+		const TArray<FName> RowNames = DataTable->GetRowNames();
+
+		if (Text.IsEmptyOrWhitespace() || !FName::IsValidXName(Text.ToString(), INVALID_NAME_CHARACTERS))
+		{
+			// Only pop up the error dialog if the rename was caused by the user's action
+			if ((CommitType == ETextCommit::OnEnter) || (CommitType == ETextCommit::OnUserMovedFocus))
+			{
+				// popup an error dialog here
+				const FText Message = FText::Format(LOCTEXT("InvalidRowName", "'{0}' is not a valid row name"), Text);
+				FMessageDialog::Open(EAppMsgType::Ok, Message);
+			}
+			return;
+		}
+		const FName NewName = DataTableUtils::MakeValidName(Text.ToString());
+		if (NewName == NAME_None)
+		{
+			// popup an error dialog here
+			const FText Message = FText::Format(LOCTEXT("InvalidRowName", "'{0}' is not a valid row name"), Text);
+			FMessageDialog::Open(EAppMsgType::Ok, Message);
+
+			return;
+		}
+		for (const FName& Name : RowNames)
+		{
+			if (Name.IsValid() && (Name == NewName))
+			{
+				//the name already exists
+				// popup an error dialog here
+				const FText Message = FText::Format(LOCTEXT("DuplicateRowName", "'{0}' is already used as a row name in this table"), Text);
+				FMessageDialog::Open(EAppMsgType::Ok, Message);
+				return;
+
+			}
+		}
+
+		const FName OldName = GetCurrentName();
+		FDataTableEditorUtils::RenameRow(DataTable, OldName, NewName);
+		DataTableEditor.Pin()->SelectionChange(DataTable, NewName);
+		*CurrentName = NewName;
+	}
+}
+
+FName SDataTableListViewRowName::GetCurrentName() const
+{
+	return CurrentName.IsValid() ? *CurrentName : NAME_None;
+
+}
+
+FText SDataTableListViewRowName::GetCurrentNameAsText() const
+{
+	return FText::FromName(GetCurrentName());
+}
+
+#undef LOCTEXT_NAMESPACE
diff --git a/Engine/Source/Editor/DataTableEditor/Private/SDataTableListViewRowName.h b/Engine/Source/Editor/DataTableEditor/Private/SDataTableListViewRowName.h
--- a/Engine/Source/Editor/DataTableEditor/Private/SDataTableListViewRowName.h
+++ b/Engine/Source/Editor/DataTableEditor/Private/SDataTableListViewRowName.h
@@ -31,9 +31,24 @@ public:
 	void Construct(const FArguments& InArgs, const TSharedRef<STableViewBase>& InOwnerTableView);
 
 	virtual FReply OnMouseButtonUp(const FGeometry& MyGeometry, const FPointerEvent& MouseEvent) override;
+
+	void OnRowRenamed(const FText& Text, ETextCommit::Type CommitType);
+
+	virtual FReply OnKeyDown(const FGeometry& MyGeometry, const FKeyEvent& InKeyEvent) override;
+
+	virtual void Tick(const FGeometry& AllottedGeometry, const double InCurrentTime, const float InDeltaTime) override;
+
+	FText GetCurrentNameAsText() const;
+	FName GetCurrentName() const;
+
 private:
 
 	void OnSearchForReferences();
+
+	TSharedPtr<SEditableText> EditableText;
+	
+	TSharedPtr<FName> CurrentName;
+
 	FDataTableEditorRowListViewDataPtr RowDataPtr;
 	TWeakPtr<FDataTableEditor> DataTableEditor;
 };
\ No newline at end of file
diff --git a/Engine/Source/Editor/DataTableEditor/Private/SRowEditor.cpp b/Engine/Source/Editor/DataTableEditor/Private/SRowEditor.cpp
--- a/Engine/Source/Editor/DataTableEditor/Private/SRowEditor.cpp
+++ b/Engine/Source/Editor/DataTableEditor/Private/SRowEditor.cpp
@@ -507,35 +507,6 @@ void SRowEditor::ConstructInternal(UDataTable* Changed)
 				]
 			]
 
-			+ SHorizontalBox::Slot()
-			[
-				SNew(SSpacer)
-			]
-			+ SHorizontalBox::Slot()
-			.AutoWidth()
-			.Padding(2)
-			.VAlign(EVerticalAlignment::VAlign_Center)
-			[
-				SNew(SBox)
-				.HAlign(HAlign_Right)
-				.Visibility(this, &SRowEditor::GetRenameVisibility)
-				[
-					SNew(STextBlock).Text(LOCTEXT("RowNameLabel", "Row Name:"))
-				]
-			]
-			+ SHorizontalBox::Slot()
-			.AutoWidth()
-			.Padding(2)
-			[
-				SNew(SBox)
-				.WidthOverride(2 * ButtonWidth)
-				.Visibility(this, &SRowEditor::GetRenameVisibility)
-				[
-					SAssignNew(RenameTextBox, SEditableTextBox)
-					.Text(this, &SRowEditor::GetCurrentNameAsText)
-					.OnTextCommitted(this, &SRowEditor::OnRowRenamed)
-				]
-			]
 			+ SHorizontalBox::Slot()
 			.AutoWidth()
 			.Padding(2)
diff --git a/Engine/Source/Editor/UnrealEd/Private/Toolkits/AssetEditorToolkit.cpp b/Engine/Source/Editor/UnrealEd/Private/Toolkits/AssetEditorToolkit.cpp
--- a/Engine/Source/Editor/UnrealEd/Private/Toolkits/AssetEditorToolkit.cpp
+++ b/Engine/Source/Editor/UnrealEd/Private/Toolkits/AssetEditorToolkit.cpp
@@ -229,9 +229,12 @@ void FAssetEditorToolkit::InitAssetEditor( const EToolkitMode::Type Mode, const
 		FExecuteAction::CreateSP( this, &FAssetEditorToolkit::FindInContentBrowser_Execute ),
 		FCanExecuteAction::CreateSP( this, &FAssetEditorToolkit::CanFindInContentBrowser ));
 		
-	ToolkitCommands->MapAction(
-		FGlobalEditorCommonCommands::Get().OpenDocumentation,
-		FExecuteAction::CreateSP( this, &FAssetEditorToolkit::BrowseDocumentation_Execute ) );
+	if (AppIdentifier != FName(TEXT("DataTableEditorApp")))
+	{
+		ToolkitCommands->MapAction(
+			FGlobalEditorCommonCommands::Get().OpenDocumentation,
+			FExecuteAction::CreateSP(this, &FAssetEditorToolkit::BrowseDocumentation_Execute));
+	}
 
 	ToolkitCommands->MapAction(
 		FAssetEditorCommonCommands::Get().ReimportAsset,
