Added forward declare and header files to accomodate for Win64 compilation errors of previous submission.

diff --git a/Engine/Source/Editor/DataTableEditor/Private/SDataTableListViewRowName.cpp b/Engine/Source/Editor/DataTableEditor/Private/SDataTableListViewRowName.cpp
--- a/Engine/Source/Editor/DataTableEditor/Private/SDataTableListViewRowName.cpp
+++ b/Engine/Source/Editor/DataTableEditor/Private/SDataTableListViewRowName.cpp
@@ -7,6 +7,8 @@
 #include "DataTableRowUtlis.h"
 #include "EditorStyleSet.h"
 #include "Framework/MultiBox/MultiBoxBuilder.h"
+#include "Widgets/Input/SEditableText.h"
+#include "Misc/MessageDialog.h"
 
 #define LOCTEXT_NAMESPACE "SDataTableListViewRowName"
 
diff --git a/Engine/Source/Editor/DataTableEditor/Private/SDataTableListViewRowName.h b/Engine/Source/Editor/DataTableEditor/Private/SDataTableListViewRowName.h
--- a/Engine/Source/Editor/DataTableEditor/Private/SDataTableListViewRowName.h
+++ b/Engine/Source/Editor/DataTableEditor/Private/SDataTableListViewRowName.h
@@ -12,7 +12,7 @@
 #include "DataTableEditorUtils.h"
 
 class FDataTableEditor;
-
+class SEditableText;
 /**
  * A widget to represent a row in a Data Table Editor widget. This widget allows us to do things like right-click
  * and take actions on a particular row of a Data Table.
