Refactored DataTable Rows for each row to be it's own widget, instead of 2 different widgets per row. Added double click editing to the Row Names column and resizing of the Row Names column.

diff --git a/Engine/Source/Editor/DataTableEditor/Private/DataTableEditor.cpp b/Engine/Source/Editor/DataTableEditor/Private/DataTableEditor.cpp
--- a/Engine/Source/Editor/DataTableEditor/Private/DataTableEditor.cpp
+++ b/Engine/Source/Editor/DataTableEditor/Private/DataTableEditor.cpp
@@ -18,7 +18,7 @@
 #include "Modules/ModuleManager.h"
 #include "Policies/PrettyJsonPrintPolicy.h"
 #include "ScopedTransaction.h"
-#include "SDataTableListViewRowName.h"
+#include "SDataTableListViewRow.h"
 #include "Serialization/JsonReader.h"
 #include "Serialization/JsonSerializer.h"
 #include "Widgets/Docking/SDockTab.h"
@@ -47,80 +47,6 @@ const FName FDataTableEditor::DataTableDetailsTabId("DataTableEditor_DataTableDe
 const FName FDataTableEditor::RowEditorTabId("DataTableEditor_RowEditor");
 const FName FDataTableEditor::RowNameColumnId("RowName");
 
-class SDataTableListViewRow : public SMultiColumnTableRow<FDataTableEditorRowListViewDataPtr>
-{
-public:
-	SLATE_BEGIN_ARGS(SDataTableListViewRow) {}
-		/** The widget that owns the tree.  We'll only keep a weak reference to it. */
-		SLATE_ARGUMENT(TSharedPtr<FDataTableEditor>, DataTableEditor)
-		/** The list item for this row */
-		SLATE_ARGUMENT(FDataTableEditorRowListViewDataPtr, Item)
-	SLATE_END_ARGS()
-
-	/** Construct function for this widget */
-	void Construct(const FArguments& InArgs, const TSharedRef<STableViewBase>& InOwnerTableView)
-	{
-		DataTableEditor = InArgs._DataTableEditor;
-		Item = InArgs._Item;
-		SMultiColumnTableRow<FDataTableEditorRowListViewDataPtr>::Construct(
-			FSuperRowType::FArguments()
-				.Style(FEditorStyle::Get(), "DataTableEditor.CellListViewRow"), 
-			InOwnerTableView
-			);
-	}
-
-
-	virtual FReply OnMouseButtonUp(const FGeometry& MyGeometry, const FPointerEvent& MouseEvent) override
-	{
-		if (MouseEvent.GetEffectingButton() == EKeys::RightMouseButton && Item.IsValid() && FEditorDelegates::OnOpenReferenceViewer.IsBound() && DataTableEditor.IsValid())
-		{
-			FDataTableEditorUtils::SelectRow(DataTableEditor.Pin()->GetDataTable(), Item->RowId);
-	
-			TSharedRef<SWidget> MenuWidget = FDataTableRowUtils::MakeRowActionsMenu(DataTableEditor.Pin(), FExecuteAction::CreateSP(this, &SDataTableListViewRow::OnSearchForReferences));
-	
-			FWidgetPath WidgetPath = MouseEvent.GetEventPath() != nullptr ? *MouseEvent.GetEventPath() : FWidgetPath();
-	
-			FSlateApplication::Get().PushMenu(AsShared(), WidgetPath, MenuWidget, MouseEvent.GetScreenSpacePosition(), FPopupTransitionEffect::ContextMenu);
-	
-			return FReply::Handled();
-		}
-	
-		return STableRow::OnMouseButtonUp(MyGeometry, MouseEvent);
-	}
-
-
-	/** Overridden from SMultiColumnTableRow.  Generates a widget for this column of the list view. */
-	virtual TSharedRef<SWidget> GenerateWidgetForColumn(const FName& ColumnName) override
-	{
-		TSharedPtr<FDataTableEditor> DataTableEditorPtr = DataTableEditor.Pin();
-		return (DataTableEditorPtr.IsValid())
-			? DataTableEditorPtr->MakeCellWidget(Item, IndexInList, ColumnName)
-			: SNullWidget::NullWidget;
-	}
-
-private:
-	void OnSearchForReferences()
-	{
-		if (DataTableEditor.IsValid() && Item.IsValid())
-		{
-			if (FDataTableEditor* DataTableEditorPtr = DataTableEditor.Pin().Get())
-			{
-				UDataTable* SourceDataTable = const_cast<UDataTable*>(DataTableEditorPtr->GetDataTable());
-
-				TArray<FAssetIdentifier> AssetIdentifiers;
-				AssetIdentifiers.Add(FAssetIdentifier(SourceDataTable, Item->RowId));
-
-				FEditorDelegates::OnOpenReferenceViewer.Broadcast(AssetIdentifiers);
-			}
-		}
-	}
-
-	/** Weak reference to the data table editor that owns our list */
-	TWeakPtr<FDataTableEditor> DataTableEditor;
-	/** The item associated with this row of data */
-	FDataTableEditorRowListViewDataPtr Item;
-};
-
 void FDataTableEditor::RegisterTabSpawners(const TSharedRef<class FTabManager>& InTabManager)
 {
 	WorkspaceMenuCategory = InTabManager->AddLocalWorkspaceMenuCategory(LOCTEXT("WorkspaceMenu_Data Table Editor", "Data Table Editor"));
@@ -384,9 +310,25 @@ FText FDataTableEditor::GetCellToolTipText(FDataTableEditorRowListViewDataPtr In
 	return TooltipText;
 }
 
-FOptionalSize FDataTableEditor::GetRowNameColumnWidth() const
+float FDataTableEditor::GetRowNameColumnWidth() const
 {
-	return FOptionalSize(RowNameColumnWidth);
+	return RowNameColumnWidth;
+}
+
+void FDataTableEditor::RefreshRowNameColumnWidth()
+{
+	
+	TSharedRef<FSlateFontMeasure> FontMeasure = FSlateApplication::Get().GetRenderer()->GetFontMeasureService();
+	const FTextBlockStyle& CellTextStyle = FEditorStyle::GetWidgetStyle<FTextBlockStyle>("DataTableEditor.CellText");
+	static const float CellPadding = 10.0f;
+
+	RowNameColumnWidth = 10.0f;
+	for (const FDataTableEditorRowListViewDataPtr& RowData : AvailableRows)
+	{
+		const float RowNameWidth = FontMeasure->Measure(RowData->DisplayName, CellTextStyle.Font).X + CellPadding;
+		RowNameColumnWidth = FMath::Max(RowNameColumnWidth, RowNameWidth);
+	}
+	
 }
 
 float FDataTableEditor::GetColumnWidth(const int32 ColumnIndex) const
@@ -430,6 +372,11 @@ void FDataTableEditor::OnColumnResized(const float NewWidth, const int32 ColumnI
 	}
 }
 
+void FDataTableEditor::OnRowNameColumnResized(const float NewWidth)
+{
+	RowNameColumnWidth = NewWidth;
+}
+
 void FDataTableEditor::LoadLayoutData()
 {
 	LayoutData.Reset();
@@ -468,19 +415,12 @@ void FDataTableEditor::SaveLayoutData()
 	}
 }
 
-TSharedRef<ITableRow> FDataTableEditor::MakeRowNameWidget(FDataTableEditorRowListViewDataPtr InRowDataPtr, const TSharedRef<STableViewBase>& OwnerTable)
-{
-	return SNew(SDataTableListViewRowName, OwnerTable)
-		.DataTableEditor(SharedThis(this))
-		.RowDataPtr(InRowDataPtr);
-}
-
 TSharedRef<ITableRow> FDataTableEditor::MakeRowWidget(FDataTableEditorRowListViewDataPtr InRowDataPtr, const TSharedRef<STableViewBase>& OwnerTable)
 {
 	return
 		SNew(SDataTableListViewRow, OwnerTable)
 		.DataTableEditor(SharedThis(this))
-		.Item(InRowDataPtr);
+		.RowDataPtr(InRowDataPtr);
 }
 
 TSharedRef<SWidget> FDataTableEditor::MakeCellWidget(FDataTableEditorRowListViewDataPtr InRowDataPtr, const int32 InRowIndex, const FName& InColumnId)
@@ -701,18 +641,7 @@ void FDataTableEditor::RefreshCachedDataTable(const FName InCachedSelection, con
 
 	// Update the desired width of the row names column
 	// This prevents it growing or shrinking as you scroll the list view
-	{
-		TSharedRef<FSlateFontMeasure> FontMeasure = FSlateApplication::Get().GetRenderer()->GetFontMeasureService();
-		const FTextBlockStyle& CellTextStyle = FEditorStyle::GetWidgetStyle<FTextBlockStyle>("DataTableEditor.CellText");
-		static const float CellPadding = 10.0f;
-		
-		RowNameColumnWidth = 10.0f;
-		for (const FDataTableEditorRowListViewDataPtr& RowData : AvailableRows)
-		{
-			const float RowNameWidth = FontMeasure->Measure(RowData->DisplayName, CellTextStyle.Font).X + CellPadding;
-			RowNameColumnWidth = FMath::Max(RowNameColumnWidth, RowNameWidth);
-		}
-	}
+	RefreshRowNameColumnWidth();
 
 	// Setup the default auto-sized columns
 	ColumnWidths.SetNum(AvailableColumns.Num());
@@ -746,6 +675,14 @@ void FDataTableEditor::RefreshCachedDataTable(const FName InCachedSelection, con
 	if (PreviousColumns != AvailableColumns)
 	{
 		ColumnNamesHeaderRow->ClearColumns();
+
+		ColumnNamesHeaderRow->AddColumn(
+			SHeaderRow::Column(RowNameColumnId)
+			.DefaultLabel(FText::GetEmpty())
+			.ManualWidth(this, &FDataTableEditor::GetRowNameColumnWidth)
+			.OnWidthChanged(this, &FDataTableEditor::OnRowNameColumnResized)
+		);
+
 		for (int32 ColumnIndex = 0; ColumnIndex < AvailableColumns.Num(); ++ColumnIndex)
 		{
 			const FDataTableEditorColumnHeaderDataPtr& ColumnData = AvailableColumns[ColumnIndex];
@@ -815,15 +752,6 @@ void FDataTableEditor::UpdateVisibleRows(const FName InCachedSelection, const bo
 			}
 		}
 	}
-	
-	// Abort restoring the cached selection if data was changed while the user is selecting a different row
-	if (RowNamesListView->GetSelectedItems() == CellsListView->GetSelectedItems())
-	{
-		RowNamesListView->RequestListRefresh();
-		CellsListView->RequestListRefresh();
-
-		RestoreCachedSelection(InCachedSelection, bUpdateEvenIfValid);
-	}
 }
 
 void FDataTableEditor::RestoreCachedSelection(const FName InCachedSelection, const bool bUpdateEvenIfValid)
@@ -861,24 +789,13 @@ TSharedRef<SVerticalBox> FDataTableEditor::CreateContentBox()
 		.Orientation(Orient_Vertical)
 		.Thickness(FVector2D(12.0f, 12.0f));
 
-	TSharedRef<SHeaderRow> RowNamesHeaderRow = SNew(SHeaderRow);
-	RowNamesHeaderRow->AddColumn(
+	ColumnNamesHeaderRow = SNew(SHeaderRow);
+	ColumnNamesHeaderRow->AddColumn(
 		SHeaderRow::Column(RowNameColumnId)
 		.DefaultLabel(FText::GetEmpty())
-		);
-
-	ColumnNamesHeaderRow = SNew(SHeaderRow);
-
-	RowNamesListView = SNew(SListView<FDataTableEditorRowListViewDataPtr>)
-		.ListItemsSource(&VisibleRows)
-		.HeaderRow(RowNamesHeaderRow)
-		.OnGenerateRow(this, &FDataTableEditor::MakeRowNameWidget)
-		.OnListViewScrolled(this, &FDataTableEditor::OnRowNamesListViewScrolled)
-		.OnSelectionChanged(this, &FDataTableEditor::OnRowSelectionChanged)
-		.ScrollbarVisibility(EVisibility::Collapsed)
-		.ConsumeMouseWheel(EConsumeMouseWheel::Always)
-		.SelectionMode(ESelectionMode::Single)
-		.AllowOverscroll(EAllowOverscroll::No);
+		.ManualWidth(this, &FDataTableEditor::GetRowNameColumnWidth)
+		.OnWidthChanged(this, &FDataTableEditor::OnRowNameColumnResized)
+	);
 
 	CellsListView = SNew(SListView<FDataTableEditorRowListViewDataPtr>)
 		.ListItemsSource(&VisibleRows)
@@ -905,15 +822,6 @@ TSharedRef<SVerticalBox> FDataTableEditor::CreateContentBox()
 		[
 			SNew(SHorizontalBox)
 			+SHorizontalBox::Slot()
-			.AutoWidth()
-			[
-				SNew(SBox)
-				.WidthOverride(this, &FDataTableEditor::GetRowNameColumnWidth)
-				[
-					RowNamesListView.ToSharedRef()
-				]
-			]
-			+SHorizontalBox::Slot()
 			[
 				SNew(SScrollBox)
 				.Orientation(Orient_Horizontal)
@@ -934,15 +842,6 @@ TSharedRef<SVerticalBox> FDataTableEditor::CreateContentBox()
 		[
 			SNew(SHorizontalBox)
 			+SHorizontalBox::Slot()
-			.AutoWidth()
-			[
-				SNew(SBox)
-				.WidthOverride(this, &FDataTableEditor::GetRowNameColumnWidth)
-				[
-					SNullWidget::NullWidget
-				]
-			]
-			+SHorizontalBox::Slot()
 			[
 				HorizontalScrollBar
 			]
@@ -1051,8 +950,6 @@ void FDataTableEditor::SetHighlightedRow(FName Name)
 	{
 		HighlightedRowName = NAME_None;
 
-		// Synchronize the list views
-		RowNamesListView->ClearSelection();
 		CellsListView->ClearSelection();
 	}
 	else
@@ -1067,14 +964,12 @@ void FDataTableEditor::SetHighlightedRow(FName Name)
 		// Synchronize the list views
 		if (NewSelectionPtr)
 		{
-			RowNamesListView->SetSelection(*NewSelectionPtr);
 			CellsListView->SetSelection(*NewSelectionPtr);
 
 			CellsListView->RequestScrollIntoView(*NewSelectionPtr);
 		}
 		else
 		{
-			RowNamesListView->ClearSelection();
 			CellsListView->ClearSelection();
 		}
 	}
diff --git a/Engine/Source/Editor/DataTableEditor/Private/DataTableEditor.h b/Engine/Source/Editor/DataTableEditor/Private/DataTableEditor.h
--- a/Engine/Source/Editor/DataTableEditor/Private/DataTableEditor.h
+++ b/Engine/Source/Editor/DataTableEditor/Private/DataTableEditor.h
@@ -27,7 +27,6 @@ class FDataTableEditor : public IDataTableEditor
 	, public FDataTableEditorUtils::INotifyOnDataTableChanged
 {
 	friend class SDataTableListViewRow;
-	friend class SDataTableListViewRowName;
 
 public:
 
@@ -113,19 +112,19 @@ protected:
 	/**	Spawns the tab with the Row Editor inside */
 	TSharedRef<SDockTab> SpawnTab_RowEditor(const FSpawnTabArgs& Args);
 
-	FOptionalSize GetRowNameColumnWidth() const;
+	float GetRowNameColumnWidth() const;
+	void RefreshRowNameColumnWidth();
 
 	float GetColumnWidth(const int32 ColumnIndex) const;
 
 	void OnColumnResized(const float NewWidth, const int32 ColumnIndex);
 
+	void OnRowNameColumnResized(const float NewWidth);
+
 	void LoadLayoutData();
 
 	void SaveLayoutData();
 
-	/** Make the widget for a row name entry in the data table row list view */
-	TSharedRef<ITableRow> MakeRowNameWidget(FDataTableEditorRowListViewDataPtr InRowDataPtr, const TSharedRef<STableViewBase>& OwnerTable);
-
 	/** Make the widget for a row entry in the data table row list view */
 	TSharedRef<ITableRow> MakeRowWidget(FDataTableEditorRowListViewDataPtr InRowDataPtr, const TSharedRef<STableViewBase>& OwnerTable);
 
diff --git a/Engine/Source/Editor/DataTableEditor/Private/SDataTableListViewRow.cpp b/Engine/Source/Editor/DataTableEditor/Private/SDataTableListViewRow.cpp
new file mode 100644
--- /dev/null
+++ b/Engine/Source/Editor/DataTableEditor/Private/SDataTableListViewRow.cpp
@@ -0,0 +1,196 @@
+// Copyright 1998-2019 Epic Games, Inc. All Rights Reserved.
+
+#include "SDataTableListViewRow.h"
+
+#include "AssetData.h"
+#include "DataTableEditor.h"
+#include "DataTableRowUtlis.h"
+#include "EditorStyleSet.h"
+#include "Framework/MultiBox/MultiBoxBuilder.h"
+#include "Widgets/Text/SInlineEditableTextBlock.h"
+#include "Misc/MessageDialog.h"
+
+#define LOCTEXT_NAMESPACE "SDataTableListViewRowName"
+
+void SDataTableListViewRow::Construct(const FArguments& InArgs, const TSharedRef<STableViewBase>& InOwnerTableView)
+{
+	RowDataPtr = InArgs._RowDataPtr;
+	CurrentName = MakeShareable(new FName(RowDataPtr->RowId));
+	DataTableEditor = InArgs._DataTableEditor;
+	SMultiColumnTableRow<FDataTableEditorRowListViewDataPtr>::Construct(
+		FSuperRowType::FArguments()
+		.Style(FEditorStyle::Get(), "DataTableEditor.CellListViewRow"),
+		InOwnerTableView
+	);
+
+}
+
+FReply SDataTableListViewRow::OnMouseButtonUp(const FGeometry& MyGeometry, const FPointerEvent& MouseEvent)
+{
+	if (MouseEvent.GetEffectingButton() == EKeys::RightMouseButton && RowDataPtr.IsValid() && FEditorDelegates::OnOpenReferenceViewer.IsBound() && DataTableEditor.IsValid())
+	{
+		FDataTableEditorUtils::SelectRow(DataTableEditor.Pin()->GetDataTable(), RowDataPtr->RowId);
+		TSharedRef<SWidget> MenuWidget = FDataTableRowUtils::MakeRowActionsMenu(DataTableEditor.Pin(), FExecuteAction::CreateSP(this, &SDataTableListViewRow::OnSearchForReferences));
+
+		FWidgetPath WidgetPath = MouseEvent.GetEventPath() != nullptr ? *MouseEvent.GetEventPath() : FWidgetPath();
+		FSlateApplication::Get().PushMenu(AsShared(), WidgetPath, MenuWidget, MouseEvent.GetScreenSpacePosition(), FPopupTransitionEffect::ContextMenu);
+		return FReply::Handled();
+	}
+
+	return STableRow::OnMouseButtonUp(MyGeometry, MouseEvent);
+}
+
+void SDataTableListViewRow::OnSearchForReferences()
+{
+	if (DataTableEditor.IsValid() && RowDataPtr.IsValid())
+	{
+		if (FDataTableEditor* DataTableEditorPtr = DataTableEditor.Pin().Get())
+		{
+			UDataTable* SourceDataTable = const_cast<UDataTable*>(DataTableEditorPtr->GetDataTable());
+
+			TArray<FAssetIdentifier> AssetIdentifiers;
+			AssetIdentifiers.Add(FAssetIdentifier(SourceDataTable, RowDataPtr->RowId));
+
+			FEditorDelegates::OnOpenReferenceViewer.Broadcast(AssetIdentifiers);
+		}
+	}
+}
+
+FReply SDataTableListViewRow::OnKeyDown(const FGeometry& MyGeometry, const FKeyEvent& InKeyEvent)
+{
+	FKey Key = InKeyEvent.GetKey();
+
+	if (Key == EKeys::Escape && InlineEditableText->HasKeyboardFocus())
+	{
+		// Clear focus
+		return FReply::Handled().SetUserFocus(SharedThis(this), EFocusCause::Cleared);
+	}
+
+	return FReply::Unhandled();
+}
+
+void SDataTableListViewRow::OnRowRenamed(const FText& Text, ETextCommit::Type CommitType)
+{
+	UDataTable* DataTable = Cast<UDataTable>(DataTableEditor.Pin()->GetEditingObject());
+
+	if (!GetCurrentNameAsText().EqualTo(Text) && DataTable)
+	{
+		const TArray<FName> RowNames = DataTable->GetRowNames();
+
+		if (Text.IsEmptyOrWhitespace() || !FName::IsValidXName(Text.ToString(), INVALID_NAME_CHARACTERS))
+		{
+			// Only pop up the error dialog if the rename was caused by the user's action
+			if ((CommitType == ETextCommit::OnEnter) || (CommitType == ETextCommit::OnUserMovedFocus))
+			{
+				// popup an error dialog here
+				const FText Message = FText::Format(LOCTEXT("InvalidRowName", "'{0}' is not a valid row name"), Text);
+				FMessageDialog::Open(EAppMsgType::Ok, Message);
+			}
+			return;
+		}
+		const FName NewName = DataTableUtils::MakeValidName(Text.ToString());
+		if (NewName == NAME_None)
+		{
+			// popup an error dialog here
+			const FText Message = FText::Format(LOCTEXT("InvalidRowName", "'{0}' is not a valid row name"), Text);
+			FMessageDialog::Open(EAppMsgType::Ok, Message);
+
+			return;
+		}
+		for (const FName& Name : RowNames)
+		{
+			if (Name.IsValid() && (Name == NewName))
+			{
+				//the name already exists
+				// popup an error dialog here
+				const FText Message = FText::Format(LOCTEXT("DuplicateRowName", "'{0}' is already used as a row name in this table"), Text);
+				FMessageDialog::Open(EAppMsgType::Ok, Message);
+				return;
+
+			}
+		}
+
+		const FName OldName = GetCurrentName();
+		FDataTableEditorUtils::RenameRow(DataTable, OldName, NewName);
+		FDataTableEditorUtils::SelectRow(DataTable, NewName);
+
+		*CurrentName = NewName;
+	}
+}
+
+TSharedRef<SWidget> SDataTableListViewRow::GenerateWidgetForColumn(const FName& ColumnName)
+{
+	TSharedPtr<FDataTableEditor> DataTableEditorPtr = DataTableEditor.Pin();
+	return (DataTableEditorPtr.IsValid())
+		? MakeCellWidget(IndexInList, ColumnName)
+		: SNullWidget::NullWidget;
+}
+
+TSharedRef<SWidget> SDataTableListViewRow::MakeCellWidget(const int32 InRowIndex, const FName& InColumnId)
+{
+	int32 ColumnIndex = 0;
+
+	FDataTableEditor* DataTableEdit = DataTableEditor.Pin().Get();
+	TArray<FDataTableEditorColumnHeaderDataPtr>& AvailableColumns = DataTableEdit->AvailableColumns;
+
+	if (InColumnId.IsEqual(FName(TEXT("RowName"))))
+	{
+		return SNew(SBox)
+			.Padding(FMargin(4, 2, 4, 2))
+			[
+				SAssignNew(InlineEditableText, SInlineEditableTextBlock)
+				.Text(RowDataPtr->DisplayName)
+				.OnTextCommitted(this, &SDataTableListViewRow::OnRowRenamed)
+				.ColorAndOpacity(DataTableEdit, &FDataTableEditor::GetRowTextColor, RowDataPtr->RowId)
+			];
+	}
+
+	for (; ColumnIndex < AvailableColumns.Num(); ++ColumnIndex)
+	{
+		const FDataTableEditorColumnHeaderDataPtr& ColumnData = AvailableColumns[ColumnIndex];
+		if (ColumnData->ColumnId == InColumnId)
+		{
+			break;
+		}
+	}
+	 
+	// Valid column ID?
+	if (AvailableColumns.IsValidIndex(ColumnIndex) && RowDataPtr->CellData.IsValidIndex(ColumnIndex))
+	{
+		return SNew(SBox)
+			.Padding(FMargin(4, 2, 4, 2))
+			[
+				SNew(STextBlock)
+				.TextStyle(FEditorStyle::Get(), "DataTableEditor.CellText")
+				.ColorAndOpacity(DataTableEdit, &FDataTableEditor::GetRowTextColor, RowDataPtr->RowId)
+				.Text(DataTableEdit, &FDataTableEditor::GetCellText, RowDataPtr, ColumnIndex)
+				.HighlightText(DataTableEdit, &FDataTableEditor::GetFilterText)
+				.ToolTipText(DataTableEdit, &FDataTableEditor::GetCellToolTipText, RowDataPtr, ColumnIndex)
+			];
+	}
+
+	return SNullWidget::NullWidget;
+}
+
+FName SDataTableListViewRow::GetCurrentName() const
+{
+	return CurrentName.IsValid() ? *CurrentName : NAME_None;
+
+}
+
+FReply SDataTableListViewRow::OnMouseButtonDoubleClick(const FGeometry& InMyGeometry, const FPointerEvent& InMouseEvent)
+{
+	if (InlineEditableText->IsHovered())
+	{
+		InlineEditableText->EnterEditingMode();
+	}
+
+	return FReply::Handled();
+}
+
+FText SDataTableListViewRow::GetCurrentNameAsText() const
+{
+	return FText::FromName(GetCurrentName());
+}
+
+#undef LOCTEXT_NAMESPACE
diff --git a/Engine/Source/Editor/DataTableEditor/Private/SDataTableListViewRow.h b/Engine/Source/Editor/DataTableEditor/Private/SDataTableListViewRow.h
new file mode 100644
--- /dev/null
+++ b/Engine/Source/Editor/DataTableEditor/Private/SDataTableListViewRow.h
@@ -0,0 +1,57 @@
+// Copyright 1998-2019 Epic Games, Inc. All Rights Reserved.
+
+#pragma once
+
+#include "CoreMinimal.h"
+#include "InputCoreTypes.h"
+#include "Input/Reply.h"
+#include "Widgets/Views/STableRow.h"
+#include "Framework/Commands/UIAction.h"
+#include "Framework/Application/MenuStack.h"
+#include "Framework/Application/SlateApplication.h"
+#include "DataTableEditorUtils.h"
+
+class FDataTableEditor;
+class SEditableText;
+/**
+ * A widget to represent a row in a Data Table Editor widget. This widget allows us to do things like right-click
+ * and take actions on a particular row of a Data Table.
+ */
+class SDataTableListViewRow : public SMultiColumnTableRow<FDataTableEditorRowListViewDataPtr>
+{
+public:
+
+	SLATE_BEGIN_ARGS(SDataTableListViewRow) {}
+	/** The owning object. This allows us access to the actual data table being edited as well as some other API functions. */
+	SLATE_ARGUMENT(TSharedPtr<FDataTableEditor>, DataTableEditor)
+		/** The row we're working with to allow us to get naming information. */
+		SLATE_ARGUMENT(FDataTableEditorRowListViewDataPtr, RowDataPtr)
+		SLATE_END_ARGS()
+
+	void Construct(const FArguments& InArgs, const TSharedRef<STableViewBase>& InOwnerTableView);
+
+	virtual FReply OnMouseButtonUp(const FGeometry& MyGeometry, const FPointerEvent& MouseEvent) override;
+
+	void OnRowRenamed(const FText& Text, ETextCommit::Type CommitType);
+
+	virtual FReply OnKeyDown(const FGeometry& MyGeometry, const FKeyEvent& InKeyEvent) override;
+
+	virtual TSharedRef<SWidget> GenerateWidgetForColumn(const FName& ColumnName) override;
+
+	FText GetCurrentNameAsText() const;
+	FName GetCurrentName() const;
+
+	virtual FReply OnMouseButtonDoubleClick(const FGeometry& InMyGeometry, const FPointerEvent& InMouseEvent);
+
+private:
+
+	void OnSearchForReferences();
+	TSharedRef<SWidget> MakeCellWidget(const int32 InRowIndex, const FName& InColumnId);
+
+	TSharedPtr<SInlineEditableTextBlock> InlineEditableText;
+
+	TSharedPtr<FName> CurrentName;
+
+	FDataTableEditorRowListViewDataPtr RowDataPtr;
+	TWeakPtr<FDataTableEditor> DataTableEditor;
+};
\ No newline at end of file
diff --git a/Engine/Source/Editor/DataTableEditor/Private/SDataTableListViewRowName.cpp b/Engine/Source/Editor/DataTableEditor/Private/SDataTableListViewRowName.cpp
deleted file mode 100644
diff --git a/Engine/Source/Editor/DataTableEditor/Private/SDataTableListViewRowName.h b/Engine/Source/Editor/DataTableEditor/Private/SDataTableListViewRowName.h
deleted file mode 100644
