Adjusted forward declare for InlineEditableText usage.

diff --git a/Engine/Source/Editor/DataTableEditor/Private/SDataTableListViewRow.h b/Engine/Source/Editor/DataTableEditor/Private/SDataTableListViewRow.h
--- a/Engine/Source/Editor/DataTableEditor/Private/SDataTableListViewRow.h
+++ b/Engine/Source/Editor/DataTableEditor/Private/SDataTableListViewRow.h
@@ -12,7 +12,7 @@
 #include "DataTableEditorUtils.h"
 
 class FDataTableEditor;
-class SEditableText;
+class SInlineEditableTextBlock;
 /**
  * A widget to represent a row in a Data Table Editor widget. This widget allows us to do things like right-click
  * and take actions on a particular row of a Data Table.
