Fixed Data Tables not updating when filled while open when using . Moved the insertion/deletion and arranging buttons from Row Editor to a bar above the search bar in the data table. Also added row renaming in the context menu, as well as the double click on the name.

diff --git a/Engine/Source/Editor/DataTableEditor/Private/DataTableEditor.cpp b/Engine/Source/Editor/DataTableEditor/Private/DataTableEditor.cpp
--- a/Engine/Source/Editor/DataTableEditor/Private/DataTableEditor.cpp
+++ b/Engine/Source/Editor/DataTableEditor/Private/DataTableEditor.cpp
@@ -37,6 +37,7 @@
 #include "UObject/StructOnScope.h"
 #include "Toolkits/GlobalEditorCommonCommands.h"
 #include "DataTableRowUtlis.h"
+#include "Engine/DataTable.h"
 
 
 
@@ -107,10 +108,12 @@ FDataTableEditor::~FDataTableEditor()
 {
 	GEditor->UnregisterForUndo(this);
 
-	const UDataTable* Table = GetDataTable();
+	UDataTable* Table = GetEditableDataTable();
 	if (Table)
 	{
 		SaveLayoutData();
+
+		Table->OnDataTableImported().RemoveAll(this);
 	}
 }
 
@@ -242,6 +245,10 @@ void FDataTableEditor::InitDataTableEditor( const EToolkitMode::Type Mode, const
 	ToolkitCommands->MapAction(FGenericCommands::Get().Copy, FExecuteAction::CreateSP(this, &FDataTableEditor::CopySelectedRow));
 	ToolkitCommands->MapAction(FGenericCommands::Get().Paste, FExecuteAction::CreateSP(this, &FDataTableEditor::PasteOnSelectedRow));
 	ToolkitCommands->MapAction(FGenericCommands::Get().Duplicate, FExecuteAction::CreateSP(this, &FDataTableEditor::DuplicateSelectedRow));
+	ToolkitCommands->MapAction(FGenericCommands::Get().Rename, FExecuteAction::CreateSP(this, &FDataTableEditor::RenameSelectedRowCommand));
+
+	Table->OnDataTableImported().AddSP(this, &FDataTableEditor::ImportDataTableUpdate);
+
 }
 
 FName FDataTableEditor::GetToolkitFName() const
@@ -259,6 +266,81 @@ FString FDataTableEditor::GetDocumentationLink() const
 	return FString(TEXT("Gameplay/DataDriven"));
 }
 
+FReply FDataTableEditor::OnAddClicked()
+{
+	UDataTable* Table = GetEditableDataTable();
+
+	if (Table)
+	{
+		//const_cast<UDataTable*>(Table)
+		
+		FName NewName = DataTableUtils::MakeValidName(TEXT("NewRow"));
+		const TArray<FName> ExisitngNames = Table->GetRowNames();
+		while (ExisitngNames.Contains(NewName))
+		{
+			NewName.SetNumber(NewName.GetNumber() + 1);
+		}
+
+		FDataTableEditorUtils::AddRow(Table, NewName);
+		FDataTableEditorUtils::SelectRow(Table, NewName);
+	}
+	return FReply::Handled();
+}
+
+FReply FDataTableEditor::OnRemoveClicked()
+{
+	UDataTable* Table = GetEditableDataTable();
+
+	if (Table)
+	{
+
+		const int32 RowToRemoveIndex = VisibleRows.IndexOfByPredicate([&](const FDataTableEditorRowListViewDataPtr& InRowName) -> bool
+		{
+			return InRowName->RowId == HighlightedRowName;
+		});
+
+		if (FDataTableEditorUtils::RemoveRow(Table, HighlightedRowName))
+		{
+			// Try and keep the same row index selected
+			const int32 RowIndexToSelect = FMath::Clamp(RowToRemoveIndex, 0, VisibleRows.Num() - 1);
+			if (VisibleRows.IsValidIndex(RowIndexToSelect))
+			{
+				FDataTableEditorUtils::SelectRow(Table, VisibleRows[RowIndexToSelect]->RowId);
+			}
+		}
+
+	}
+	return FReply::Handled();
+}
+
+FReply FDataTableEditor::OnMoveRowClicked(FDataTableEditorUtils::ERowMoveDirection MoveDirection)
+{
+	UDataTable* Table = GetEditableDataTable();
+
+	if (Table)
+	{
+		FDataTableEditorUtils::MoveRow(Table, HighlightedRowName, MoveDirection);
+	}
+	return FReply::Handled();
+}
+
+FReply FDataTableEditor::OnMoveToExtentClicked(FDataTableEditorUtils::ERowMoveDirection MoveDirection)
+{
+	UDataTable* Table = GetEditableDataTable();
+
+	if (Table)
+	{
+		// We move by the row map size, as FDataTableEditorUtils::MoveRow will automatically clamp this as appropriate
+		FDataTableEditorUtils::MoveRow(Table, HighlightedRowName, MoveDirection, Table->GetRowMap().Num());
+	}
+	return FReply::Handled();
+}
+
+UDataTable* FDataTableEditor::GetEditableDataTable() const
+{
+	return Cast<UDataTable>(GetEditingObject());
+}
+
 FText FDataTableEditor::GetBaseToolkitName() const
 {
 	return LOCTEXT( "AppLabel", "DataTable Editor" );
@@ -453,18 +535,6 @@ TSharedRef<SWidget> FDataTableEditor::MakeCellWidget(FDataTableEditorRowListView
 	return SNullWidget::NullWidget;
 }
 
-void FDataTableEditor::OnRowNamesListViewScrolled(double InScrollOffset)
-{
-	// Synchronize the list views
-	CellsListView->SetScrollOffset(InScrollOffset);
-}
-
-void FDataTableEditor::OnCellsListViewScrolled(double InScrollOffset)
-{
-	// Synchronize the list views
-	RowNamesListView->SetScrollOffset(InScrollOffset);
-}
-
 void FDataTableEditor::OnRowSelectionChanged(FDataTableEditorRowListViewDataPtr InNewSelection, ESelectInfo::Type InSelectInfo)
 {
 	const bool bSelectionChanged = !InNewSelection.IsValid() || InNewSelection->RowId != HighlightedRowName;
@@ -537,6 +607,26 @@ void FDataTableEditor::DuplicateSelectedRow()
 	FDataTableEditorUtils::SelectRow(TablePtr, NewName);
 }
 
+void FDataTableEditor::RenameSelectedRowCommand()
+{
+	UDataTable* TablePtr = Cast<UDataTable>(GetEditingObject());
+	FName NewName = HighlightedRowName; 
+
+	if (NewName == NAME_None || TablePtr == nullptr)
+		return;
+
+	for (const FDataTableEditorRowListViewDataPtr& RowData : VisibleRows)
+	{
+		if (RowData->RowId == NewName)
+		{
+			TSharedPtr< SDataTableListViewRow > RowWidget = StaticCastSharedPtr< SDataTableListViewRow >(CellsListView->WidgetFromItem(RowData));
+			RowWidget->SetRowForRename();
+		}
+	}
+
+
+}
+
 FText FDataTableEditor::GetFilterText() const
 {
 	return ActiveFilterText;
@@ -715,6 +805,16 @@ void FDataTableEditor::RefreshCachedDataTable(const FName InCachedSelection, con
 	}
 }
 
+void FDataTableEditor::ImportDataTableUpdate()
+{
+	UDataTable* Table = GetEditableDataTable();
+
+	if (Table)
+	{
+		FDataTableEditorUtils::FDataTableEditorManager::Get().PostChange(Table, FDataTableEditorUtils::EDataTableChangeInfo::RowList);
+	}
+}
+
 void FDataTableEditor::UpdateVisibleRows(const FName InCachedSelection, const bool bUpdateEvenIfValid)
 {
 	if (ActiveFilterText.IsEmptyOrWhitespace())
@@ -801,7 +901,6 @@ TSharedRef<SVerticalBox> FDataTableEditor::CreateContentBox()
 		.ListItemsSource(&VisibleRows)
 		.HeaderRow(ColumnNamesHeaderRow)
 		.OnGenerateRow(this, &FDataTableEditor::MakeRowWidget)
-		.OnListViewScrolled(this, &FDataTableEditor::OnCellsListViewScrolled)
 		.OnSelectionChanged(this, &FDataTableEditor::OnRowSelectionChanged)
 		.ExternalScrollbar(VerticalScrollBar)
 		.ConsumeMouseWheel(EConsumeMouseWheel::Always)
@@ -811,6 +910,111 @@ TSharedRef<SVerticalBox> FDataTableEditor::CreateContentBox()
 	RefreshCachedDataTable();
 
 	return SNew(SVerticalBox)
+		+ SVerticalBox::Slot()
+		.AutoHeight()
+		[
+			SNew(SHorizontalBox)
+			+ SHorizontalBox::Slot()
+			.AutoWidth()
+			.Padding(2)
+			[
+				SNew(SButton)
+				.ButtonStyle(FEditorStyle::Get(), "HoverHintOnly")
+				.ForegroundColor(FSlateColor::UseForeground())
+				.HAlign(HAlign_Center)
+				.VAlign(VAlign_Center)
+				.OnClicked(this, &FDataTableEditor::OnAddClicked)
+				.ToolTipText(LOCTEXT("AddRowTooltip", "Add a new row to the data table"))
+				[
+					SNew(SImage)
+					.Image(FEditorStyle::Get().GetBrush("Plus"))
+				]
+			]
+			+ SHorizontalBox::Slot()
+			.AutoWidth()
+			.Padding(2)
+			[
+				SNew(SButton)
+				.ButtonStyle(FEditorStyle::Get(), "HoverHintOnly")
+				.ForegroundColor(FSlateColor::UseForeground())
+				.HAlign(HAlign_Center)
+				.VAlign(VAlign_Center)
+				.OnClicked(this, &FDataTableEditor::OnRemoveClicked)
+				.ToolTipText(LOCTEXT("RemoveRowTooltip", "Remove the currently selected row from the data table"))
+				[
+					SNew(SImage)
+					.Image(FEditorStyle::Get().GetBrush("Cross"))
+				]
+			]
+			+ SHorizontalBox::Slot()
+			.AutoWidth()
+			.Padding(2)
+			[
+				SNew(SButton)
+				.ButtonStyle(FEditorStyle::Get(), "HoverHintOnly")
+				.ForegroundColor(FSlateColor::UseForeground())
+				.HAlign(HAlign_Center)
+				.VAlign(VAlign_Center)
+				.OnClicked(this, &FDataTableEditor::OnMoveRowClicked, FDataTableEditorUtils::ERowMoveDirection::Up)
+				.ToolTipText(LOCTEXT("MoveUpTooltip", "Move the currently selected row up by one in the data table"))
+				[
+					SNew(STextBlock)
+					.Font(FEditorStyle::Get().GetFontStyle("FontAwesome.14"))
+					.Text(FText::FromString(FString(TEXT("\xf106"))) /*fa-angle-up*/)
+				]
+			]
+			+ SHorizontalBox::Slot()
+			.AutoWidth()
+			.Padding(2)
+			[
+				SNew(SButton)
+				.ButtonStyle(FEditorStyle::Get(), "HoverHintOnly")
+				.ForegroundColor(FSlateColor::UseForeground())
+				.HAlign(HAlign_Center)
+				.VAlign(VAlign_Center)
+				.OnClicked(this, &FDataTableEditor::OnMoveRowClicked, FDataTableEditorUtils::ERowMoveDirection::Down)
+				.ToolTipText(LOCTEXT("MoveDownTooltip", "Move the currently selected row down by one in the data table"))
+				[
+					SNew(STextBlock)
+					.Font(FEditorStyle::Get().GetFontStyle("FontAwesome.14"))
+					.Text(FText::FromString(FString(TEXT("\xf107"))) /*fa-angle-down*/)
+				]
+			]
+			+ SHorizontalBox::Slot()
+			.AutoWidth()
+			.Padding(2)
+			[
+				SNew(SButton)
+				.ButtonStyle(FEditorStyle::Get(), "HoverHintOnly")
+				.ForegroundColor(FSlateColor::UseForeground())
+				.HAlign(HAlign_Center)
+				.VAlign(VAlign_Center)
+				.OnClicked(this, &FDataTableEditor::OnMoveToExtentClicked, FDataTableEditorUtils::ERowMoveDirection::Up)
+				.ToolTipText(LOCTEXT("MoveToTopTooltip", "Move the currently selected row to the top of the data table"))
+				[
+					SNew(STextBlock)
+					.Font(FEditorStyle::Get().GetFontStyle("FontAwesome.14"))
+					.Text(FText::FromString(FString(TEXT("\xf102"))) /*fa-angle-double-up*/)
+				]
+			]
+			+ SHorizontalBox::Slot()
+			.AutoWidth()
+			.Padding(2)
+			[
+				SNew(SButton)
+				.ButtonStyle(FEditorStyle::Get(), "HoverHintOnly")
+				.ForegroundColor(FSlateColor::UseForeground())
+				.HAlign(HAlign_Center)
+				.VAlign(VAlign_Center)
+				.OnClicked(this, &FDataTableEditor::OnMoveToExtentClicked, FDataTableEditorUtils::ERowMoveDirection::Down)
+				.ToolTipText(LOCTEXT("MoveToBottomTooltip", "Move the currently selected row to the bottom of the data table"))
+				[
+					SNew(STextBlock)
+					.Font(FEditorStyle::Get().GetFontStyle("FontAwesome.14"))
+					.Text(FText::FromString(FString(TEXT("\xf103"))) /*fa-angle-double-down*/)
+				]
+			]
+		]
 		+SVerticalBox::Slot()
 		.AutoHeight()
 		[
diff --git a/Engine/Source/Editor/DataTableEditor/Private/DataTableEditor.h b/Engine/Source/Editor/DataTableEditor/Private/DataTableEditor.h
--- a/Engine/Source/Editor/DataTableEditor/Private/DataTableEditor.h
+++ b/Engine/Source/Editor/DataTableEditor/Private/DataTableEditor.h
@@ -83,6 +83,8 @@ protected:
 
 	void RefreshCachedDataTable(const FName InCachedSelection = NAME_None, const bool bUpdateEvenIfValid = false);
 
+	void ImportDataTableUpdate();
+
 	void UpdateVisibleRows(const FName InCachedSelection = NAME_None, const bool bUpdateEvenIfValid = false);
 
 	void RestoreCachedSelection(const FName InCachedSelection, const bool bUpdateEvenIfValid = false);
@@ -131,15 +133,12 @@ protected:
 	/** Make the widget for a cell entry in the data table row list view */
 	TSharedRef<SWidget> MakeCellWidget(FDataTableEditorRowListViewDataPtr InRowDataPtr, const int32 InRowIndex, const FName& InColumnId);
 
-	void OnRowNamesListViewScrolled(double InScrollOffset);
-
-	void OnCellsListViewScrolled(double InScrollOffset);
-
 	void OnRowSelectionChanged(FDataTableEditorRowListViewDataPtr InNewSelection, ESelectInfo::Type InSelectInfo);
 
 	void CopySelectedRow();
 	void PasteOnSelectedRow();
 	void DuplicateSelectedRow();
+	void RenameSelectedRowCommand();
 
 	/** Helper function for creating and registering the tab containing the data table data */
 	virtual void CreateAndRegisterDataTableTab(const TSharedRef<class FTabManager>& InTabManager);
@@ -153,6 +152,15 @@ protected:
 	void BrowseDocumentation_Execute() const;
 
 	virtual FString GetDocumentationLink() const override;
+	
+	FReply OnAddClicked();
+	FReply OnRemoveClicked();
+	FReply OnMoveRowClicked(FDataTableEditorUtils::ERowMoveDirection MoveDirection);
+	FReply OnMoveToExtentClicked(FDataTableEditorUtils::ERowMoveDirection MoveDirection);
+
+
+private:
+	UDataTable* GetEditableDataTable() const;
 
 protected:
 
@@ -192,9 +200,6 @@ protected:
 	/** Header row containing entries for each column in AvailableColumns */
 	TSharedPtr<SHeaderRow> ColumnNamesHeaderRow;
 
-	/** List view responsible for showing the row names column */
-	TSharedPtr<SListView<FDataTableEditorRowListViewDataPtr>> RowNamesListView;
-
 	/** List view responsible for showing the rows in VisibleRows for each entry in AvailableColumns */
 	TSharedPtr<SListView<FDataTableEditorRowListViewDataPtr>> CellsListView;
 
diff --git a/Engine/Source/Editor/DataTableEditor/Private/SDataTableListViewRow.cpp b/Engine/Source/Editor/DataTableEditor/Private/SDataTableListViewRow.cpp
--- a/Engine/Source/Editor/DataTableEditor/Private/SDataTableListViewRow.cpp
+++ b/Engine/Source/Editor/DataTableEditor/Private/SDataTableListViewRow.cpp
@@ -40,6 +40,22 @@ FReply SDataTableListViewRow::OnMouseButtonUp(const FGeometry& MyGeometry, const
 	return STableRow::OnMouseButtonUp(MyGeometry, MouseEvent);
 }
 
+FReply SDataTableListViewRow::OnMouseButtonDown(const FGeometry& MyGeometry, const FPointerEvent& MouseEvent)
+{
+	if (MouseEvent.GetEffectingButton() == EKeys::LeftMouseButton && RowDataPtr.IsValid() && DataTableEditor.IsValid())
+	{
+		FDataTableEditorUtils::SelectRow(DataTableEditor.Pin()->GetDataTable(), RowDataPtr->RowId);
+		//TSharedRef<SWidget> MenuWidget = FDataTableRowUtils::MakeRowActionsMenu(DataTableEditor.Pin(), FExecuteAction::CreateSP(this, &SDataTableListViewRow::OnSearchForReferences));
+		//
+		//FWidgetPath WidgetPath = MouseEvent.GetEventPath() != nullptr ? *MouseEvent.GetEventPath() : FWidgetPath();
+		//FSlateApplication::Get().PushMenu(AsShared(), WidgetPath, MenuWidget, MouseEvent.GetScreenSpacePosition(), FPopupTransitionEffect::ContextMenu);
+		return FReply::Handled();
+	}
+
+	return STableRow::OnMouseButtonDown(MyGeometry, MouseEvent);
+
+}
+
 void SDataTableListViewRow::OnSearchForReferences()
 {
 	if (DataTableEditor.IsValid() && RowDataPtr.IsValid())
@@ -188,6 +204,11 @@ FReply SDataTableListViewRow::OnMouseButtonDoubleClick(const FGeometry& InMyGeom
 	return FReply::Handled();
 }
 
+void SDataTableListViewRow::SetRowForRename()
+{
+	InlineEditableText->EnterEditingMode();
+}
+
 FText SDataTableListViewRow::GetCurrentNameAsText() const
 {
 	return FText::FromName(GetCurrentName());
diff --git a/Engine/Source/Editor/DataTableEditor/Private/SDataTableListViewRow.h b/Engine/Source/Editor/DataTableEditor/Private/SDataTableListViewRow.h
--- a/Engine/Source/Editor/DataTableEditor/Private/SDataTableListViewRow.h
+++ b/Engine/Source/Editor/DataTableEditor/Private/SDataTableListViewRow.h
@@ -31,6 +31,7 @@ public:
 	void Construct(const FArguments& InArgs, const TSharedRef<STableViewBase>& InOwnerTableView);
 
 	virtual FReply OnMouseButtonUp(const FGeometry& MyGeometry, const FPointerEvent& MouseEvent) override;
+	virtual FReply OnMouseButtonDown(const FGeometry& MyGeometry, const FPointerEvent& MouseEvent) override;
 
 	void OnRowRenamed(const FText& Text, ETextCommit::Type CommitType);
 
@@ -43,6 +44,8 @@ public:
 
 	virtual FReply OnMouseButtonDoubleClick(const FGeometry& InMyGeometry, const FPointerEvent& InMouseEvent);
 
+	void SetRowForRename();
+
 private:
 
 	void OnSearchForReferences();
diff --git a/Engine/Source/Editor/DataTableEditor/Private/SRowEditor.cpp b/Engine/Source/Editor/DataTableEditor/Private/SRowEditor.cpp
--- a/Engine/Source/Editor/DataTableEditor/Private/SRowEditor.cpp
+++ b/Engine/Source/Editor/DataTableEditor/Private/SRowEditor.cpp
@@ -438,40 +438,6 @@ void SRowEditor::ConstructInternal(UDataTable* Changed)
 			+ SHorizontalBox::Slot()
 			.AutoWidth()
 			.Padding(2)
-			[
-				SNew(SButton)
-				.ButtonStyle(FEditorStyle::Get(), "HoverHintOnly")
-				.ForegroundColor(FSlateColor::UseForeground())
-				.HAlign(HAlign_Center)
-				.VAlign(VAlign_Center)
-				.OnClicked(this, &SRowEditor::OnAddClicked)
-				.IsEnabled(this, &SRowEditor::IsAddRowEnabled)
-				.ToolTipText(LOCTEXT("AddRowTooltip", "Add a new row to the data table"))
-				[
-					SNew(SImage)
-					.Image(FEditorStyle::Get().GetBrush("Plus"))
-				]
-			]
-			+ SHorizontalBox::Slot()
-			.AutoWidth()
-			.Padding(2)
-			[
-				SNew(SButton)
-				.ButtonStyle(FEditorStyle::Get(), "HoverHintOnly")
-				.ForegroundColor(FSlateColor::UseForeground())
-				.HAlign(HAlign_Center)
-				.VAlign(VAlign_Center)
-				.OnClicked(this, &SRowEditor::OnRemoveClicked)
-				.IsEnabled(this, &SRowEditor::IsRemoveRowEnabled)
-				.ToolTipText(LOCTEXT("RemoveRowTooltip", "Remove the currently selected row from the data table"))
-				[
-					SNew(SImage)
-					.Image(FEditorStyle::Get().GetBrush("Cross"))
-				]
-			]
-			+ SHorizontalBox::Slot()
-			.AutoWidth()
-			.Padding(2)
 			[
 				SNew(SBox)
 				.WidthOverride(2 * ButtonWidth)
@@ -507,78 +473,6 @@ void SRowEditor::ConstructInternal(UDataTable* Changed)
 				]
 			]
 
-			+ SHorizontalBox::Slot()
-			.AutoWidth()
-			.Padding(2)
-			[
-				SNew(SButton)
-				.ButtonStyle(FEditorStyle::Get(), "HoverHintOnly")
-				.ForegroundColor(FSlateColor::UseForeground())
-				.HAlign(HAlign_Center)
-				.VAlign(VAlign_Center)
-				.OnClicked(this, &SRowEditor::OnMoveRowClicked, FDataTableEditorUtils::ERowMoveDirection::Up)
-				.IsEnabled(this, &SRowEditor::IsMoveRowUpEnabled)
-				.ToolTipText(LOCTEXT("MoveUpTooltip", "Move the currently selected row up by one in the data table"))
-				[
-					SNew(STextBlock)
-					.Font(FEditorStyle::Get().GetFontStyle("FontAwesome.14"))
-					.Text(FText::FromString(FString(TEXT("\xf106"))) /*fa-angle-up*/)
-				]
-			]
-			+ SHorizontalBox::Slot()
-			.AutoWidth()
-			.Padding(2)
-			[
-				SNew(SButton)
-				.ButtonStyle(FEditorStyle::Get(), "HoverHintOnly")
-				.ForegroundColor(FSlateColor::UseForeground())
-				.HAlign(HAlign_Center)
-				.VAlign(VAlign_Center)
-				.OnClicked(this, &SRowEditor::OnMoveRowClicked, FDataTableEditorUtils::ERowMoveDirection::Down)
-				.IsEnabled(this, &SRowEditor::IsMoveRowDownEnabled)
-				.ToolTipText(LOCTEXT("MoveDownTooltip", "Move the currently selected row down by one in the data table"))
-				[
-					SNew(STextBlock)
-					.Font(FEditorStyle::Get().GetFontStyle("FontAwesome.14"))
-					.Text(FText::FromString(FString(TEXT("\xf107"))) /*fa-angle-down*/)
-				]
-			]
-			+ SHorizontalBox::Slot()
-			.AutoWidth()
-			.Padding(2)
-			[
-				SNew(SButton)
-				.ButtonStyle(FEditorStyle::Get(), "HoverHintOnly")
-				.ForegroundColor(FSlateColor::UseForeground())
-				.HAlign(HAlign_Center)
-				.VAlign(VAlign_Center)
-				.OnClicked(this, &SRowEditor::OnMoveToExtentClicked, FDataTableEditorUtils::ERowMoveDirection::Up)
-				.IsEnabled(this, &SRowEditor::IsMoveRowUpEnabled)
-				.ToolTipText(LOCTEXT("MoveToTopTooltip", "Move the currently selected row to the top of the data table"))
-				[
-					SNew(STextBlock)
-					.Font(FEditorStyle::Get().GetFontStyle("FontAwesome.14"))
-					.Text(FText::FromString(FString(TEXT("\xf102"))) /*fa-angle-double-up*/)
-				]
-			]
-			+ SHorizontalBox::Slot()
-			.AutoWidth()
-			.Padding(2)
-			[
-				SNew(SButton)
-				.ButtonStyle(FEditorStyle::Get(), "HoverHintOnly")
-				.ForegroundColor(FSlateColor::UseForeground())
-				.HAlign(HAlign_Center)
-				.VAlign(VAlign_Center)
-				.OnClicked(this, &SRowEditor::OnMoveToExtentClicked, FDataTableEditorUtils::ERowMoveDirection::Down)
-				.IsEnabled(this, &SRowEditor::IsMoveRowDownEnabled)
-				.ToolTipText(LOCTEXT("MoveToBottomTooltip", "Move the currently selected row to the bottom of the data table"))
-				[
-					SNew(STextBlock)
-					.Font(FEditorStyle::Get().GetFontStyle("FontAwesome.14"))
-					.Text(FText::FromString(FString(TEXT("\xf103"))) /*fa-angle-double-down*/)
-				]
-			]
 		]
 		+ SVerticalBox::Slot()
 		[
diff --git a/Engine/Source/Editor/DataTableEditor/Public/DataTableRowUtlis.cpp b/Engine/Source/Editor/DataTableEditor/Public/DataTableRowUtlis.cpp
--- a/Engine/Source/Editor/DataTableEditor/Public/DataTableRowUtlis.cpp
+++ b/Engine/Source/Editor/DataTableEditor/Public/DataTableRowUtlis.cpp
@@ -22,6 +22,7 @@ TSharedRef<SWidget> FDataTableRowUtils::MakeRowActionsMenu(TSharedPtr<IDataTable
 		MenuBuilder.AddMenuEntry(FGenericCommands::Get().Copy);
 		MenuBuilder.AddMenuEntry(FGenericCommands::Get().Paste);
 		MenuBuilder.AddMenuEntry(FGenericCommands::Get().Duplicate);
+		MenuBuilder.AddMenuEntry(FGenericCommands::Get().Rename);
 		MenuBuilder.AddMenuSeparator();
 		MenuBuilder.AddMenuEntry(SearchForReferencesActionName, SearchForReferencesActionTooltip, 
 			FSlateIcon(), FUIAction(SearchForReferencesAction));
diff --git a/Engine/Source/Runtime/Engine/Classes/Engine/DataTable.h b/Engine/Source/Runtime/Engine/Classes/Engine/DataTable.h
--- a/Engine/Source/Runtime/Engine/Classes/Engine/DataTable.h
+++ b/Engine/Source/Runtime/Engine/Classes/Engine/DataTable.h
@@ -59,6 +59,7 @@ class ENGINE_VTABLE UDataTable
 	GENERATED_UCLASS_BODY()
 
 	DECLARE_MULTICAST_DELEGATE(FOnDataTableChanged);
+	DECLARE_MULTICAST_DELEGATE(FOnDataTableImport);
 	
 	friend FDataTableEditorUtils;
 	friend UGameplayTagTableManager;
@@ -148,10 +149,17 @@ private:
 	/** A multicast delegate that is called any time the data table changes. */
 	FOnDataTableChanged OnDataTableChangedDelegate;
 
+	/** A multicast delegate that is called any time a data table is imported. */
+	FOnDataTableChanged OnDataTableImportedDelegate;
+
 public:
 	/** Gets a multicast delegate that is called any time the data table changes. */
 	FOnDataTableChanged& OnDataTableChanged() { return OnDataTableChangedDelegate; }
 
+	/** Gets a multicast delegate that is called any time a data table is imported. */
+	FOnDataTableImport& OnDataTableImported() { return OnDataTableImportedDelegate; }
+
+
 	//~ Begin UDataTable Interface
 
 	/** Get all of the rows in the table, regardless of name */
diff --git a/Engine/Source/Runtime/Engine/Private/DataTable.cpp b/Engine/Source/Runtime/Engine/Private/DataTable.cpp
--- a/Engine/Source/Runtime/Engine/Private/DataTable.cpp
+++ b/Engine/Source/Runtime/Engine/Private/DataTable.cpp
@@ -189,6 +189,7 @@ void UDataTable::OnPostDataImported(TArray<FString>& OutCollectedImportProblems)
 		}
 	}
 
+	OnDataTableImported().Broadcast();
 	OnDataTableChanged().Broadcast();
 }
 
