Added in Copy and Paste buttons for Data Tables and Delete on the context menu for Data Tables, as well as some optimizations and stylistic changes.

diff --git a/Engine/Source/Editor/DataTableEditor/Private/DataTableEditor.cpp b/Engine/Source/Editor/DataTableEditor/Private/DataTableEditor.cpp
--- a/Engine/Source/Editor/DataTableEditor/Private/DataTableEditor.cpp
+++ b/Engine/Source/Editor/DataTableEditor/Private/DataTableEditor.cpp
@@ -246,6 +246,7 @@ void FDataTableEditor::InitDataTableEditor( const EToolkitMode::Type Mode, const
 	ToolkitCommands->MapAction(FGenericCommands::Get().Paste, FExecuteAction::CreateSP(this, &FDataTableEditor::PasteOnSelectedRow));
 	ToolkitCommands->MapAction(FGenericCommands::Get().Duplicate, FExecuteAction::CreateSP(this, &FDataTableEditor::DuplicateSelectedRow));
 	ToolkitCommands->MapAction(FGenericCommands::Get().Rename, FExecuteAction::CreateSP(this, &FDataTableEditor::RenameSelectedRowCommand));
+	ToolkitCommands->MapAction(FGenericCommands::Get().Delete, FExecuteAction::CreateSP(this, &FDataTableEditor::DeleteSelectedRow));
 
 	Table->OnDataTableImported().AddSP(this, &FDataTableEditor::ImportDataTableUpdate);
 
@@ -271,12 +272,9 @@ FReply FDataTableEditor::OnAddClicked()
 	UDataTable* Table = GetEditableDataTable();
 
 	if (Table)
-	{
-		//const_cast<UDataTable*>(Table)
-		
+	{		
 		FName NewName = DataTableUtils::MakeValidName(TEXT("NewRow"));
-		const TArray<FName> ExisitngNames = Table->GetRowNames();
-		while (ExisitngNames.Contains(NewName))
+		while (Table->GetRowMap().Contains(NewName))
 		{
 			NewName.SetNumber(NewName.GetNumber() + 1);
 		}
@@ -284,6 +282,7 @@ FReply FDataTableEditor::OnAddClicked()
 		FDataTableEditorUtils::AddRow(Table, NewName);
 		FDataTableEditorUtils::SelectRow(Table, NewName);
 	}
+
 	return FReply::Handled();
 }
 
@@ -336,6 +335,26 @@ FReply FDataTableEditor::OnMoveToExtentClicked(FDataTableEditorUtils::ERowMoveDi
 	return FReply::Handled();
 }
 
+FReply FDataTableEditor::OnCopyClicked()
+{
+	UDataTable* Table = GetEditableDataTable();
+	if (Table)
+	{
+		CopySelectedRow();
+	}
+	return FReply::Handled();
+}
+
+FReply FDataTableEditor::OnPasteClicked()
+{
+	UDataTable* Table = GetEditableDataTable();
+	if (Table)
+	{
+		PasteOnSelectedRow();
+	}
+	return FReply::Handled();
+}
+
 UDataTable* FDataTableEditor::GetEditableDataTable() const
 {
 	return Cast<UDataTable>(GetEditingObject());
@@ -595,7 +614,9 @@ void FDataTableEditor::DuplicateSelectedRow()
 	FName NewName = HighlightedRowName;
 
 	if (NewName == NAME_None || TablePtr == nullptr)
+	{
 		return;
+	}
 
 	const TArray<FName> ExistingNames = TablePtr->GetRowNames();
 	while (ExistingNames.Contains(NewName))
@@ -613,18 +634,27 @@ void FDataTableEditor::RenameSelectedRowCommand()
 	FName NewName = HighlightedRowName; 
 
 	if (NewName == NAME_None || TablePtr == nullptr)
+	{
 		return;
+	}
 
-	for (const FDataTableEditorRowListViewDataPtr& RowData : VisibleRows)
+	if (VisibleRows.IsValidIndex(HighlightedVisibleRowIndex))
 	{
-		if (RowData->RowId == NewName)
-		{
-			TSharedPtr< SDataTableListViewRow > RowWidget = StaticCastSharedPtr< SDataTableListViewRow >(CellsListView->WidgetFromItem(RowData));
-			RowWidget->SetRowForRename();
-		}
+		TSharedPtr< SDataTableListViewRow > RowWidget = StaticCastSharedPtr< SDataTableListViewRow >(CellsListView->WidgetFromItem(VisibleRows[HighlightedVisibleRowIndex]));
+		RowWidget->SetRowForRename();
 	}
+}
 
+void FDataTableEditor::DeleteSelectedRow()
+{
+	UDataTable* TablePtr = Cast<UDataTable>(GetEditingObject());
+
+	if (HighlightedRowName == NAME_None || TablePtr == nullptr)
+	{
+		return;
+	}
 
+	FDataTableEditorUtils::RemoveRow(TablePtr, HighlightedRowName);
 }
 
 FText FDataTableEditor::GetFilterText() const
@@ -1014,6 +1044,38 @@ TSharedRef<SVerticalBox> FDataTableEditor::CreateContentBox()
 					.Text(FText::FromString(FString(TEXT("\xf103"))) /*fa-angle-double-down*/)
 				]
 			]
+			+ SHorizontalBox::Slot()
+				.AutoWidth()
+				.Padding(2)
+				[
+				SNew(SButton)
+				.ButtonStyle(FEditorStyle::Get(), "HoverHintOnly")
+				.ForegroundColor(FSlateColor::UseForeground())
+				.HAlign(HAlign_Center)
+				.VAlign(VAlign_Center)
+				.OnClicked(this, &FDataTableEditor::OnCopyClicked)
+				.ToolTipText(LOCTEXT("CopyTooltip", "Copy the currently selected row"))
+				[
+					SNew(SImage)
+					.Image(FEditorStyle::Get().GetBrush("DataTableEditor.Copy"))
+				]
+			]
+			+ SHorizontalBox::Slot()
+				.AutoWidth()
+				.Padding(2)
+				[
+				SNew(SButton)
+				.ButtonStyle(FEditorStyle::Get(), "HoverHintOnly")
+				.ForegroundColor(FSlateColor::UseForeground())
+				.HAlign(HAlign_Center)
+				.VAlign(VAlign_Center)
+				.OnClicked(this, &FDataTableEditor::OnPasteClicked)
+				.ToolTipText(LOCTEXT("PasteTooltip", "Paste on the currently selected row"))
+				[
+					SNew(SImage)
+					.Image(FEditorStyle::Get().GetBrush("DataTableEditor.Paste"))
+				]
+			]
 		]
 		+SVerticalBox::Slot()
 		.AutoHeight()
@@ -1153,17 +1215,24 @@ void FDataTableEditor::SetHighlightedRow(FName Name)
 	if (Name.IsNone())
 	{
 		HighlightedRowName = NAME_None;
-
 		CellsListView->ClearSelection();
+		HighlightedVisibleRowIndex = -1;
 	}
 	else
 	{
 		HighlightedRowName = Name;
 
-		FDataTableEditorRowListViewDataPtr* NewSelectionPtr = VisibleRows.FindByPredicate([&Name](const FDataTableEditorRowListViewDataPtr& RowData) -> bool
+		FDataTableEditorRowListViewDataPtr* NewSelectionPtr = NULL;
+		for (HighlightedVisibleRowIndex = 0; HighlightedVisibleRowIndex < VisibleRows.Num(); ++HighlightedVisibleRowIndex)
 		{
-			return RowData->RowId == Name;
-		});
+			if (VisibleRows[HighlightedVisibleRowIndex]->RowId == Name)
+			{
+				NewSelectionPtr = &(VisibleRows[HighlightedVisibleRowIndex]);
+
+				break;
+			}
+		}
+
 
 		// Synchronize the list views
 		if (NewSelectionPtr)
diff --git a/Engine/Source/Editor/DataTableEditor/Private/DataTableEditor.h b/Engine/Source/Editor/DataTableEditor/Private/DataTableEditor.h
--- a/Engine/Source/Editor/DataTableEditor/Private/DataTableEditor.h
+++ b/Engine/Source/Editor/DataTableEditor/Private/DataTableEditor.h
@@ -139,6 +139,7 @@ protected:
 	void PasteOnSelectedRow();
 	void DuplicateSelectedRow();
 	void RenameSelectedRowCommand();
+	void DeleteSelectedRow();
 
 	/** Helper function for creating and registering the tab containing the data table data */
 	virtual void CreateAndRegisterDataTableTab(const TSharedRef<class FTabManager>& InTabManager);
@@ -157,7 +158,8 @@ protected:
 	FReply OnRemoveClicked();
 	FReply OnMoveRowClicked(FDataTableEditorUtils::ERowMoveDirection MoveDirection);
 	FReply OnMoveToExtentClicked(FDataTableEditorUtils::ERowMoveDirection MoveDirection);
-
+	FReply OnCopyClicked();
+	FReply OnPasteClicked();
 
 private:
 	UDataTable* GetEditableDataTable() const;
@@ -215,6 +217,9 @@ protected:
 	/** The name of the currently selected row */
 	FName HighlightedRowName;
 
+	/** The visible row index of the currently selected row */
+	int32 HighlightedVisibleRowIndex;
+
 	/** The current filter text applied to the data table */
 	FText ActiveFilterText;
 
diff --git a/Engine/Source/Editor/DataTableEditor/Private/SDataTableListViewRow.cpp b/Engine/Source/Editor/DataTableEditor/Private/SDataTableListViewRow.cpp
--- a/Engine/Source/Editor/DataTableEditor/Private/SDataTableListViewRow.cpp
+++ b/Engine/Source/Editor/DataTableEditor/Private/SDataTableListViewRow.cpp
@@ -40,22 +40,6 @@ FReply SDataTableListViewRow::OnMouseButtonUp(const FGeometry& MyGeometry, const
 	return STableRow::OnMouseButtonUp(MyGeometry, MouseEvent);
 }
 
-FReply SDataTableListViewRow::OnMouseButtonDown(const FGeometry& MyGeometry, const FPointerEvent& MouseEvent)
-{
-	if (MouseEvent.GetEffectingButton() == EKeys::LeftMouseButton && RowDataPtr.IsValid() && DataTableEditor.IsValid())
-	{
-		FDataTableEditorUtils::SelectRow(DataTableEditor.Pin()->GetDataTable(), RowDataPtr->RowId);
-		//TSharedRef<SWidget> MenuWidget = FDataTableRowUtils::MakeRowActionsMenu(DataTableEditor.Pin(), FExecuteAction::CreateSP(this, &SDataTableListViewRow::OnSearchForReferences));
-		//
-		//FWidgetPath WidgetPath = MouseEvent.GetEventPath() != nullptr ? *MouseEvent.GetEventPath() : FWidgetPath();
-		//FSlateApplication::Get().PushMenu(AsShared(), WidgetPath, MenuWidget, MouseEvent.GetScreenSpacePosition(), FPopupTransitionEffect::ContextMenu);
-		return FReply::Handled();
-	}
-
-	return STableRow::OnMouseButtonDown(MyGeometry, MouseEvent);
-
-}
-
 void SDataTableListViewRow::OnSearchForReferences()
 {
 	if (DataTableEditor.IsValid() && RowDataPtr.IsValid())
diff --git a/Engine/Source/Editor/DataTableEditor/Private/SDataTableListViewRow.h b/Engine/Source/Editor/DataTableEditor/Private/SDataTableListViewRow.h
--- a/Engine/Source/Editor/DataTableEditor/Private/SDataTableListViewRow.h
+++ b/Engine/Source/Editor/DataTableEditor/Private/SDataTableListViewRow.h
@@ -31,7 +31,6 @@ public:
 	void Construct(const FArguments& InArgs, const TSharedRef<STableViewBase>& InOwnerTableView);
 
 	virtual FReply OnMouseButtonUp(const FGeometry& MyGeometry, const FPointerEvent& MouseEvent) override;
-	virtual FReply OnMouseButtonDown(const FGeometry& MyGeometry, const FPointerEvent& MouseEvent) override;
 
 	void OnRowRenamed(const FText& Text, ETextCommit::Type CommitType);
 
diff --git a/Engine/Source/Editor/DataTableEditor/Public/DataTableRowUtlis.cpp b/Engine/Source/Editor/DataTableEditor/Public/DataTableRowUtlis.cpp
--- a/Engine/Source/Editor/DataTableEditor/Public/DataTableRowUtlis.cpp
+++ b/Engine/Source/Editor/DataTableEditor/Public/DataTableRowUtlis.cpp
@@ -23,6 +23,7 @@ TSharedRef<SWidget> FDataTableRowUtils::MakeRowActionsMenu(TSharedPtr<IDataTable
 		MenuBuilder.AddMenuEntry(FGenericCommands::Get().Paste);
 		MenuBuilder.AddMenuEntry(FGenericCommands::Get().Duplicate);
 		MenuBuilder.AddMenuEntry(FGenericCommands::Get().Rename);
+		MenuBuilder.AddMenuEntry(FGenericCommands::Get().Delete);
 		MenuBuilder.AddMenuSeparator();
 		MenuBuilder.AddMenuEntry(SearchForReferencesActionName, SearchForReferencesActionTooltip, 
 			FSlateIcon(), FUIAction(SearchForReferencesAction));
diff --git a/Engine/Source/Editor/EditorStyle/Private/SlateEditorStyle.cpp b/Engine/Source/Editor/EditorStyle/Private/SlateEditorStyle.cpp
--- a/Engine/Source/Editor/EditorStyle/Private/SlateEditorStyle.cpp
+++ b/Engine/Source/Editor/EditorStyle/Private/SlateEditorStyle.cpp
@@ -6840,6 +6840,9 @@ void FSlateEditorStyle::FStyle::SetupToolkitStyles()
 	{
 		Set( "DataTableEditor.Tabs.Properties", new IMAGE_BRUSH( "/Icons/icon_tab_SelectionDetails_16x", Icon16x16 ) );
 
+		Set("DataTableEditor.Copy", new IMAGE_BRUSH("/Icons/Edit/icon_Edit_Copy_40x", Icon16x16));
+		Set("DataTableEditor.Paste", new IMAGE_BRUSH("/Icons/Edit/icon_Edit_Paste_40x", Icon16x16));
+
 		Set( "DataTableEditor.CellText", FTextBlockStyle(NormalText)
 			.SetFont( DEFAULT_FONT("Regular", 9 ))
 			);
