Added in insertion of row into the context menu of Data Tables. Added in a separator for arrangement and insert/copy/paste/delete buttons. Added in duplicate button to Data Tables.

diff --git a/Engine/Source/Editor/DataTableEditor/Private/DataTableEditor.cpp b/Engine/Source/Editor/DataTableEditor/Private/DataTableEditor.cpp
--- a/Engine/Source/Editor/DataTableEditor/Private/DataTableEditor.cpp
+++ b/Engine/Source/Editor/DataTableEditor/Private/DataTableEditor.cpp
@@ -32,6 +32,7 @@
 #include "Widgets/Images/SImage.h"
 #include "Widgets/Input/SHyperlink.h"
 #include "Widgets/Notifications/SNotificationList.h"
+#include "Widgets/Layout/SSeparator.h"
 #include "SourceCodeNavigation.h"
 #include "PropertyEditorModule.h"
 #include "UObject/StructOnScope.h"
@@ -355,6 +356,16 @@
 	return FReply::Handled();
 }
 
+FReply FDataTableEditor::OnDuplicateClicked()
+{
+	UDataTable* Table = GetEditableDataTable();
+	if (Table)
+	{
+		DuplicateSelectedRow();
+	}
+	return FReply::Handled();
+}
+
 UDataTable* FDataTableEditor::GetEditableDataTable() const
 {
 	return Cast<UDataTable>(GetEditingObject());
@@ -920,6 +931,7 @@
 		.Thickness(FVector2D(12.0f, 12.0f));
 
 	ColumnNamesHeaderRow = SNew(SHeaderRow);
+
 	ColumnNamesHeaderRow->AddColumn(
 		SHeaderRow::Column(RowNameColumnId)
 		.DefaultLabel(FText::GetEmpty())
@@ -966,3 +978,51 @@
 			[
 				SNew(SButton)
 				.ButtonStyle(FEditorStyle::Get(), "HoverHintOnly")
+				.ForegroundColor(FSlateColor::UseForeground())
+				.HAlign(HAlign_Center)
+				.VAlign(VAlign_Center)
+				.OnClicked(this, &FDataTableEditor::OnCopyClicked)
+				.ToolTipText(LOCTEXT("CopyTooltip", "Copy the currently selected row"))
+				[
+					SNew(SImage)
+					.Image(FEditorStyle::Get().GetBrush("DataTableEditor.Copy"))
+				]
+			]
+			+ SHorizontalBox::Slot()
+			.AutoWidth()
+			.Padding(2)
+			[
+				SNew(SButton)
+				.ButtonStyle(FEditorStyle::Get(), "HoverHintOnly")
+				.ForegroundColor(FSlateColor::UseForeground())
+				.HAlign(HAlign_Center)
+				.VAlign(VAlign_Center)
+				.OnClicked(this, &FDataTableEditor::OnPasteClicked)
+				.ToolTipText(LOCTEXT("PasteTooltip", "Paste on the currently selected row"))
+				[
+					SNew(SImage)
+					.Image(FEditorStyle::Get().GetBrush("DataTableEditor.Paste"))
+				]
+			]
+			+ SHorizontalBox::Slot()
+			.AutoWidth()
+			.Padding(2)
+			[
+				SNew(SButton)
+				.ButtonStyle(FEditorStyle::Get(), "HoverHintOnly")
+				.ForegroundColor(FSlateColor::UseForeground())
+				.HAlign(HAlign_Center)
+				.VAlign(VAlign_Center)
+				.OnClicked(this, &FDataTableEditor::OnDuplicateClicked)
+				.ToolTipText(LOCTEXT("DuplicateTooltip", "Duplicate the currently selected row"))
+				[
+					SNew(SImage)
+					.Image(FEditorStyle::Get().GetBrush("DataTableEditor.Duplicate"))
+				]
+			]
+			+ SHorizontalBox::Slot()
+			.AutoWidth()
+			.Padding(2)
+			[
+				SNew(SButton)
+				.ButtonStyle(FEditorStyle::Get(), "HoverHintOnly")
@@ -980,5 +1040,12 @@
 			.AutoWidth()
 			.Padding(2)
 			[
+				SNew(SSeparator)
+				.Orientation(Orient_Vertical)
+			]
+			+ SHorizontalBox::Slot()
+			.AutoWidth()
+			.Padding(2)
+			[
 				SNew(SButton)
 				.ButtonStyle(FEditorStyle::Get(), "HoverHintOnly")
@@ -1044,38 +1111,6 @@
 					.Text(FText::FromString(FString(TEXT("\xf103"))) /*fa-angle-double-down*/)
 				]
 			]
-			+ SHorizontalBox::Slot()
-				.AutoWidth()
-				.Padding(2)
-				[
-				SNew(SButton)
-				.ButtonStyle(FEditorStyle::Get(), "HoverHintOnly")
-				.ForegroundColor(FSlateColor::UseForeground())
-				.HAlign(HAlign_Center)
-				.VAlign(VAlign_Center)
-				.OnClicked(this, &FDataTableEditor::OnCopyClicked)
-				.ToolTipText(LOCTEXT("CopyTooltip", "Copy the currently selected row"))
-				[
-					SNew(SImage)
-					.Image(FEditorStyle::Get().GetBrush("DataTableEditor.Copy"))
-				]
-			]
-			+ SHorizontalBox::Slot()
-				.AutoWidth()
-				.Padding(2)
-				[
-				SNew(SButton)
-				.ButtonStyle(FEditorStyle::Get(), "HoverHintOnly")
-				.ForegroundColor(FSlateColor::UseForeground())
-				.HAlign(HAlign_Center)
-				.VAlign(VAlign_Center)
-				.OnClicked(this, &FDataTableEditor::OnPasteClicked)
-				.ToolTipText(LOCTEXT("PasteTooltip", "Paste on the currently selected row"))
-				[
-					SNew(SImage)
-					.Image(FEditorStyle::Get().GetBrush("DataTableEditor.Paste"))
-				]
-			]
 		]
 		+SVerticalBox::Slot()
 		.AutoHeight()
diff --git a/Engine/Source/Editor/DataTableEditor/Private/DataTableEditor.h b/Engine/Source/Editor/DataTableEditor/Private/DataTableEditor.h
--- a/Engine/Source/Editor/DataTableEditor/Private/DataTableEditor.h
+++ b/Engine/Source/Editor/DataTableEditor/Private/DataTableEditor.h
@@ -160,6 +160,7 @@ protected:
 	FReply OnMoveToExtentClicked(FDataTableEditorUtils::ERowMoveDirection MoveDirection);
 	FReply OnCopyClicked();
 	FReply OnPasteClicked();
+	FReply OnDuplicateClicked();
 
 private:
 	UDataTable* GetEditableDataTable() const;
diff --git a/Engine/Source/Editor/DataTableEditor/Private/SDataTableListViewRow.cpp b/Engine/Source/Editor/DataTableEditor/Private/SDataTableListViewRow.cpp
--- a/Engine/Source/Editor/DataTableEditor/Private/SDataTableListViewRow.cpp
+++ b/Engine/Source/Editor/DataTableEditor/Private/SDataTableListViewRow.cpp
@@ -30,7 +30,9 @@ FReply SDataTableListViewRow::OnMouseButtonUp(const FGeometry& MyGeometry, const
 	if (MouseEvent.GetEffectingButton() == EKeys::RightMouseButton && RowDataPtr.IsValid() && FEditorDelegates::OnOpenReferenceViewer.IsBound() && DataTableEditor.IsValid())
 	{
 		FDataTableEditorUtils::SelectRow(DataTableEditor.Pin()->GetDataTable(), RowDataPtr->RowId);
-		TSharedRef<SWidget> MenuWidget = FDataTableRowUtils::MakeRowActionsMenu(DataTableEditor.Pin(), FExecuteAction::CreateSP(this, &SDataTableListViewRow::OnSearchForReferences));
+		TSharedRef<SWidget> MenuWidget = FDataTableRowUtils::MakeRowActionsMenu(DataTableEditor.Pin(),
+			FExecuteAction::CreateSP(this, &SDataTableListViewRow::OnSearchForReferences),
+			FExecuteAction::CreateSP(this, &SDataTableListViewRow::OnInsertNewRow));
 
 		FWidgetPath WidgetPath = MouseEvent.GetEventPath() != nullptr ? *MouseEvent.GetEventPath() : FWidgetPath();
 		FSlateApplication::Get().PushMenu(AsShared(), WidgetPath, MenuWidget, MouseEvent.GetScreenSpacePosition(), FPopupTransitionEffect::ContextMenu);
@@ -56,6 +58,29 @@ void SDataTableListViewRow::OnSearchForReferences()
 	}
 }
 
+void SDataTableListViewRow::OnInsertNewRow()
+{
+	if (DataTableEditor.IsValid() && RowDataPtr.IsValid())
+	{
+		if (FDataTableEditor* DataTableEditorPtr = DataTableEditor.Pin().Get())
+		{
+			UDataTable* SourceDataTable = const_cast<UDataTable*>(DataTableEditorPtr->GetDataTable());
+
+			if (SourceDataTable)
+			{
+				FName NewName = DataTableUtils::MakeValidName(TEXT("NewRow"));
+				while (SourceDataTable->GetRowMap().Contains(NewName))
+				{
+					NewName.SetNumber(NewName.GetNumber() + 1);
+				}
+
+				FDataTableEditorUtils::AddRow(SourceDataTable, NewName);
+				FDataTableEditorUtils::SelectRow(SourceDataTable, NewName);
+			}
+		}
+	}
+}
+
 FReply SDataTableListViewRow::OnKeyDown(const FGeometry& MyGeometry, const FKeyEvent& InKeyEvent)
 {
 	FKey Key = InKeyEvent.GetKey();
diff --git a/Engine/Source/Editor/DataTableEditor/Private/SDataTableListViewRow.h b/Engine/Source/Editor/DataTableEditor/Private/SDataTableListViewRow.h
--- a/Engine/Source/Editor/DataTableEditor/Private/SDataTableListViewRow.h
+++ b/Engine/Source/Editor/DataTableEditor/Private/SDataTableListViewRow.h
@@ -48,6 +48,8 @@ public:
 private:
 
 	void OnSearchForReferences();
+	void OnInsertNewRow();
+
 	TSharedRef<SWidget> MakeCellWidget(const int32 InRowIndex, const FName& InColumnId);
 
 	TSharedPtr<SInlineEditableTextBlock> InlineEditableText;
diff --git a/Engine/Source/Editor/DataTableEditor/Public/DataTableRowUtlis.cpp b/Engine/Source/Editor/DataTableEditor/Public/DataTableRowUtlis.cpp
--- a/Engine/Source/Editor/DataTableEditor/Public/DataTableRowUtlis.cpp
+++ b/Engine/Source/Editor/DataTableEditor/Public/DataTableRowUtlis.cpp
@@ -14,11 +14,16 @@
 const FText FDataTableRowUtils::SearchForReferencesActionName = LOCTEXT("FDataTableRowUtils_SearchForReferences", "Find Row References");
 const FText FDataTableRowUtils::SearchForReferencesActionTooltip = LOCTEXT("FDataTableRowUtils_SearchForReferencesTooltip", "Find assets that reference this Row");
 
-TSharedRef<SWidget> FDataTableRowUtils::MakeRowActionsMenu(TSharedPtr<IDataTableEditor> Editor, FExecuteAction SearchForReferencesAction)
+const FText InsertNewRowActionName = LOCTEXT("FDataTableRowUtils_InsertNewRow", "Insert New Row");
+const FText InsertNewRowActionTooltip = LOCTEXT("FDataTableRowUtils_InsertNewRowTooltip", "Insert a new Row");
+
+TSharedRef<SWidget> FDataTableRowUtils::MakeRowActionsMenu(TSharedPtr<IDataTableEditor> Editor, FExecuteAction SearchForReferencesAction, FExecuteAction InsertNewRowAction)
 {
 	if (SearchForReferencesAction.IsBound())
 	{
 		FMenuBuilder MenuBuilder(true, Editor->GetToolkitCommands());
+		MenuBuilder.AddMenuEntry(InsertNewRowActionName, InsertNewRowActionTooltip,
+			FSlateIcon(FEditorStyle::GetStyleSetName(), "Plus"), FUIAction(InsertNewRowAction));
 		MenuBuilder.AddMenuEntry(FGenericCommands::Get().Copy);
 		MenuBuilder.AddMenuEntry(FGenericCommands::Get().Paste);
 		MenuBuilder.AddMenuEntry(FGenericCommands::Get().Duplicate);
diff --git a/Engine/Source/Editor/DataTableEditor/Public/DataTableRowUtlis.h b/Engine/Source/Editor/DataTableEditor/Public/DataTableRowUtlis.h
--- a/Engine/Source/Editor/DataTableEditor/Public/DataTableRowUtlis.h
+++ b/Engine/Source/Editor/DataTableEditor/Public/DataTableRowUtlis.h
@@ -12,7 +12,7 @@ class FDetailWidgetRow;
 class DATATABLEEDITOR_API FDataTableRowUtils
 {
 public:
-	static TSharedRef<SWidget> MakeRowActionsMenu(TSharedPtr<class IDataTableEditor> Editor, FExecuteAction SearchForReferencesAction);
+	static TSharedRef<SWidget> MakeRowActionsMenu(TSharedPtr<class IDataTableEditor> Editor, FExecuteAction SearchForReferencesAction, FExecuteAction InsertNewRowAction);
 	static void AddSearchForReferencesContextMenu(FDetailWidgetRow& RowNameDetailWidget, FExecuteAction SearchForReferencesAction);
 
 private:
diff --git a/Engine/Source/Editor/EditorStyle/Private/SlateEditorStyle.cpp b/Engine/Source/Editor/EditorStyle/Private/SlateEditorStyle.cpp
--- a/Engine/Source/Editor/EditorStyle/Private/SlateEditorStyle.cpp
+++ b/Engine/Source/Editor/EditorStyle/Private/SlateEditorStyle.cpp
@@ -6842,6 +6842,7 @@ void FSlateEditorStyle::FStyle::SetupToolkitStyles()
 
 		Set("DataTableEditor.Copy", new IMAGE_BRUSH("/Icons/Edit/icon_Edit_Copy_40x", Icon16x16));
 		Set("DataTableEditor.Paste", new IMAGE_BRUSH("/Icons/Edit/icon_Edit_Paste_40x", Icon16x16));
+		Set("DataTableEditor.Duplicate", new IMAGE_BRUSH("Icons/Edit/icon_Edit_Duplicate_40x", Icon16x16));
 
 		Set( "DataTableEditor.CellText", FTextBlockStyle(NormalText)
 			.SetFont( DEFAULT_FONT("Regular", 9 ))
