Added a Row Number column to Data Tables. Added Sorting of Data Tables by column. Fixed Clearing/Resetting of Data Table search bar to reset data table selection.

diff --git a/Engine/Source/Editor/DataTableEditor/Private/DataTableEditor.cpp b/Engine/Source/Editor/DataTableEditor/Private/DataTableEditor.cpp
--- a/Engine/Source/Editor/DataTableEditor/Private/DataTableEditor.cpp
+++ b/Engine/Source/Editor/DataTableEditor/Private/DataTableEditor.cpp
@@ -48,6 +48,7 @@ const FName FDataTableEditor::DataTableTabId("DataTableEditor_DataTable");
 const FName FDataTableEditor::DataTableDetailsTabId("DataTableEditor_DataTableDetails");
 const FName FDataTableEditor::RowEditorTabId("DataTableEditor_RowEditor");
 const FName FDataTableEditor::RowNameColumnId("RowName");
+const FName FDataTableEditor::RowNumberColumnId("RowNumber");
 
 void FDataTableEditor::RegisterTabSpawners(const TSharedRef<class FTabManager>& InTabManager)
 {
@@ -366,6 +367,114 @@ FReply FDataTableEditor::OnDuplicateClicked()
 	return FReply::Handled();
 }
 
+EColumnSortMode::Type FDataTableEditor::GetColumnSortMode(const FName ColumnId) const
+{
+	if (SortByColumn != ColumnId)
+	{
+		return EColumnSortMode::None;
+	}
+
+	return SortMode;
+}
+
+void FDataTableEditor::OnColumnSortModeChanged(const EColumnSortPriority::Type SortPriority, const FName& ColumnId, const EColumnSortMode::Type InSortMode)
+{
+	int32 ColumnIndex;
+
+	SortMode = InSortMode;
+	SortByColumn = ColumnId;
+
+	for (ColumnIndex = 0; ColumnIndex < AvailableColumns.Num(); ++ColumnIndex)
+	{
+		if (AvailableColumns[ColumnIndex]->ColumnId == ColumnId)
+		{
+			break;
+		}
+	}
+
+	if (AvailableColumns.IsValidIndex(ColumnIndex))
+	{
+		if (InSortMode == EColumnSortMode::Ascending)
+		{
+			VisibleRows.Sort([ColumnIndex](const FDataTableEditorRowListViewDataPtr& first, const FDataTableEditorRowListViewDataPtr& second)
+			{					
+				int32 Result = (first->CellData[ColumnIndex].ToString()).Compare(second->CellData[ColumnIndex].ToString());
+
+				if (!Result)
+				{
+					return first->RowNum < second->RowNum;
+
+				}
+
+				return Result < 0;
+			});
+
+		}
+		else if (InSortMode == EColumnSortMode::Descending)
+		{
+			VisibleRows.Sort([ColumnIndex](const FDataTableEditorRowListViewDataPtr& first, const FDataTableEditorRowListViewDataPtr& second)
+			{
+				int32 Result = (first->CellData[ColumnIndex].ToString()).Compare(second->CellData[ColumnIndex].ToString());
+
+				if (!Result)
+				{
+					return first->RowNum > second->RowNum;
+				}
+
+				return Result > 0;
+			});
+		}
+	}
+
+	CellsListView->RequestListRefresh();
+}
+
+void FDataTableEditor::OnColumnNumberSortModeChanged(const EColumnSortPriority::Type SortPriority, const FName& ColumnId, const EColumnSortMode::Type InSortMode)
+{
+	SortMode = InSortMode;
+	SortByColumn = ColumnId;
+
+	if (InSortMode == EColumnSortMode::Ascending)
+	{
+		VisibleRows.Sort([](const FDataTableEditorRowListViewDataPtr& first, const FDataTableEditorRowListViewDataPtr& second)
+		{
+			return first->RowNum < second->RowNum;
+		});
+	}
+	else if (InSortMode == EColumnSortMode::Descending)
+	{
+		VisibleRows.Sort([](const FDataTableEditorRowListViewDataPtr& first, const FDataTableEditorRowListViewDataPtr& second)
+		{
+			return first->RowNum > second->RowNum;
+		});
+	}
+
+	CellsListView->RequestListRefresh();
+}
+
+void FDataTableEditor::OnColumnNameSortModeChanged(const EColumnSortPriority::Type SortPriority, const FName& ColumnId, const EColumnSortMode::Type InSortMode)
+{
+	SortMode = InSortMode;
+	SortByColumn = ColumnId;
+
+	if (InSortMode == EColumnSortMode::Ascending)
+	{
+		VisibleRows.Sort([](const FDataTableEditorRowListViewDataPtr& first, const FDataTableEditorRowListViewDataPtr& second)
+		{
+			return (first->DisplayName).ToString() < (second->DisplayName).ToString();
+		});
+	}
+	else if (InSortMode == EColumnSortMode::Descending)
+	{
+		VisibleRows.Sort([](const FDataTableEditorRowListViewDataPtr& first, const FDataTableEditorRowListViewDataPtr& second)
+		{
+			return (first->DisplayName).ToString() > (second->DisplayName).ToString();
+		});
+	}
+
+	CellsListView->RequestListRefresh();
+}
+
 UDataTable* FDataTableEditor::GetEditableDataTable() const
 {
 	return Cast<UDataTable>(GetEditingObject());
@@ -675,8 +784,34 @@ FText FDataTableEditor::GetFilterText() const
 
 void FDataTableEditor::OnFilterTextChanged(const FText& InFilterText)
 {
-	ActiveFilterText = InFilterText;
-	UpdateVisibleRows();
+	if (InFilterText.IsEmpty())
+	{
+		OnFilterCleared();
+	}
+	else
+	{
+		ActiveFilterText = InFilterText;
+		UpdateVisibleRows();
+	}
+}
+
+void FDataTableEditor::OnFilterTextCommitted(const FText& NewText, ETextCommit::Type CommitInfo)
+{
+	if (CommitInfo == ETextCommit::OnCleared)
+	{
+		OnFilterCleared();
+	}
+}
+
+void FDataTableEditor::OnFilterCleared()
+{
+	ActiveFilterText = FText();
+	if (VisibleRows.IsValidIndex(HighlightedVisibleRowIndex))
+	{
+		VisibleRows = AvailableRows;
+		SearchBoxWidget->SetText(ActiveFilterText);
+		CellsListView->RequestListRefresh();
+	}
 }
 
 void FDataTableEditor::PostRegenerateMenusAndToolbars()
@@ -807,11 +942,20 @@ void FDataTableEditor::RefreshCachedDataTable(const FName InCachedSelection, con
 	{
 		ColumnNamesHeaderRow->ClearColumns();
 
+		ColumnNamesHeaderRow->AddColumn(
+			SHeaderRow::Column(RowNumberColumnId)
+			.DefaultLabel(FText::GetEmpty())
+			.SortMode(this, &FDataTableEditor::GetColumnSortMode, RowNumberColumnId)
+			.OnSort(this, &FDataTableEditor::OnColumnNumberSortModeChanged)
+		);
+
 		ColumnNamesHeaderRow->AddColumn(
 			SHeaderRow::Column(RowNameColumnId)
 			.DefaultLabel(FText::GetEmpty())
 			.ManualWidth(this, &FDataTableEditor::GetRowNameColumnWidth)
 			.OnWidthChanged(this, &FDataTableEditor::OnRowNameColumnResized)
+			.SortMode(this, &FDataTableEditor::GetColumnSortMode, RowNameColumnId)
+			.OnSort(this, &FDataTableEditor::OnColumnNameSortModeChanged)
 		);
 
 		for (int32 ColumnIndex = 0; ColumnIndex < AvailableColumns.Num(); ++ColumnIndex)
@@ -823,6 +967,8 @@ void FDataTableEditor::RefreshCachedDataTable(const FName InCachedSelection, con
 				.DefaultLabel(ColumnData->DisplayName)
 				.ManualWidth(TAttribute<float>::Create(TAttribute<float>::FGetter::CreateSP(this, &FDataTableEditor::GetColumnWidth, ColumnIndex)))
 				.OnWidthChanged(this, &FDataTableEditor::OnColumnResized, ColumnIndex)
+				.SortMode(this, &FDataTableEditor::GetColumnSortMode, ColumnData->ColumnId)
+				.OnSort(this, &FDataTableEditor::OnColumnSortModeChanged)
 				[
 					SNew(SBox)
 					.Padding(FMargin(0, 4, 0, 4))
@@ -932,13 +1078,6 @@ TSharedRef<SVerticalBox> FDataTableEditor::CreateContentBox()
 
 	ColumnNamesHeaderRow = SNew(SHeaderRow);
 
-	ColumnNamesHeaderRow->AddColumn(
-		SHeaderRow::Column(RowNameColumnId)
-		.DefaultLabel(FText::GetEmpty())
-		.ManualWidth(this, &FDataTableEditor::GetRowNameColumnWidth)
-		.OnWidthChanged(this, &FDataTableEditor::OnRowNameColumnResized)
-	);
-
 	CellsListView = SNew(SListView<FDataTableEditorRowListViewDataPtr>)
 		.ListItemsSource(&VisibleRows)
 		.HeaderRow(ColumnNamesHeaderRow)
@@ -1053,7 +1192,7 @@ TSharedRef<SVerticalBox> FDataTableEditor::CreateContentBox()
 				.HAlign(HAlign_Center)
 				.VAlign(VAlign_Center)
 				.OnClicked(this, &FDataTableEditor::OnMoveRowClicked, FDataTableEditorUtils::ERowMoveDirection::Up)
-				.ToolTipText(LOCTEXT("MoveUpTooltip", "Move the currently selected row up by one in the data table"))
+				.ToolTipText(LOCTEXT("MoveUpTooltip", "Move the currently selected row up by one in the data table sorted by row number"))
 				[
 					SNew(STextBlock)
 					.Font(FEditorStyle::Get().GetFontStyle("FontAwesome.14"))
@@ -1070,7 +1209,7 @@ TSharedRef<SVerticalBox> FDataTableEditor::CreateContentBox()
 				.HAlign(HAlign_Center)
 				.VAlign(VAlign_Center)
 				.OnClicked(this, &FDataTableEditor::OnMoveRowClicked, FDataTableEditorUtils::ERowMoveDirection::Down)
-				.ToolTipText(LOCTEXT("MoveDownTooltip", "Move the currently selected row down by one in the data table"))
+				.ToolTipText(LOCTEXT("MoveDownTooltip", "Move the currently selected row down by one in the data table sorted by row number"))
 				[
 					SNew(STextBlock)
 					.Font(FEditorStyle::Get().GetFontStyle("FontAwesome.14"))
@@ -1087,7 +1226,7 @@ TSharedRef<SVerticalBox> FDataTableEditor::CreateContentBox()
 				.HAlign(HAlign_Center)
 				.VAlign(VAlign_Center)
 				.OnClicked(this, &FDataTableEditor::OnMoveToExtentClicked, FDataTableEditorUtils::ERowMoveDirection::Up)
-				.ToolTipText(LOCTEXT("MoveToTopTooltip", "Move the currently selected row to the top of the data table"))
+				.ToolTipText(LOCTEXT("MoveToTopTooltip", "Move the currently selected row to the top of the data table sorted by row number"))
 				[
 					SNew(STextBlock)
 					.Font(FEditorStyle::Get().GetFontStyle("FontAwesome.14"))
@@ -1104,20 +1243,20 @@ TSharedRef<SVerticalBox> FDataTableEditor::CreateContentBox()
 				.HAlign(HAlign_Center)
 				.VAlign(VAlign_Center)
 				.OnClicked(this, &FDataTableEditor::OnMoveToExtentClicked, FDataTableEditorUtils::ERowMoveDirection::Down)
-				.ToolTipText(LOCTEXT("MoveToBottomTooltip", "Move the currently selected row to the bottom of the data table"))
+				.ToolTipText(LOCTEXT("MoveToBottomTooltip", "Move the currently selected row to the bottom of the data table based on row number"))
 				[
 					SNew(STextBlock)
 					.Font(FEditorStyle::Get().GetFontStyle("FontAwesome.14"))
 					.Text(FText::FromString(FString(TEXT("\xf103"))) /*fa-angle-double-down*/)
 				]
 			]
-		]
-		+SVerticalBox::Slot()
-		.AutoHeight()
-		[
-			SNew(SSearchBox)
-			.InitialText(this, &FDataTableEditor::GetFilterText)
-			.OnTextChanged(this, &FDataTableEditor::OnFilterTextChanged)
+			+ SHorizontalBox::Slot()
+			[
+				SAssignNew(SearchBoxWidget, SSearchBox)
+				.InitialText(this, &FDataTableEditor::GetFilterText)
+				.OnTextChanged(this, &FDataTableEditor::OnFilterTextChanged)
+				.OnTextCommitted(this, &FDataTableEditor::OnFilterTextCommitted)
+			]
 		]
 		+SVerticalBox::Slot()
 		[
diff --git a/Engine/Source/Editor/DataTableEditor/Private/DataTableEditor.h b/Engine/Source/Editor/DataTableEditor/Private/DataTableEditor.h
--- a/Engine/Source/Editor/DataTableEditor/Private/DataTableEditor.h
+++ b/Engine/Source/Editor/DataTableEditor/Private/DataTableEditor.h
@@ -15,6 +15,7 @@
 #include "Kismet2/StructureEditorUtils.h"
 #include "DataTableEditorUtils.h"
 #include "SRowEditor.h"
+#include "Widgets/Input/SSearchBox.h"
 
 class FJsonObject;
 
@@ -91,6 +92,10 @@ protected:
 	
 	void OnFilterTextChanged(const FText& InFilterText);
 
+	void OnFilterTextCommitted(const FText& NewText, ETextCommit::Type CommitInfo);
+
+	void OnFilterCleared();
+
 	virtual void PostRegenerateMenusAndToolbars() override;
 
 	FReply OnFindRowInContentBrowserClicked();
@@ -162,6 +167,11 @@ protected:
 	FReply OnPasteClicked();
 	FReply OnDuplicateClicked();
 
+	EColumnSortMode::Type GetColumnSortMode(const FName ColumnId) const;
+	void OnColumnSortModeChanged(const EColumnSortPriority::Type SortPriority, const FName& ColumnId, const EColumnSortMode::Type InSortMode);
+	void OnColumnNumberSortModeChanged(const EColumnSortPriority::Type SortPriority, const FName& ColumnId, const EColumnSortMode::Type InSortMode);
+	void OnColumnNameSortModeChanged(const EColumnSortPriority::Type SortPriority, const FName& ColumnId, const EColumnSortMode::Type InSortMode);
+
 private:
 	UDataTable* GetEditableDataTable() const;
 
@@ -188,6 +198,9 @@ protected:
 	/** Property viewing widget */
 	TSharedPtr<class IDetailsView> PropertyView;
 
+	/** UI for the "Row Editor" tab */
+	TSharedPtr<SSearchBox> SearchBoxWidget;
+
 	/** UI for the "Row Editor" tab */
 	TSharedPtr<SWidget> RowEditorTabWidget;
 
@@ -196,7 +209,7 @@ protected:
 
 	/** Array of the rows that are available for editing */
 	TArray<FDataTableEditorRowListViewDataPtr> AvailableRows;
-	
+
 	/** Array of the rows that match the active filter(s) */
 	TArray<FDataTableEditorRowListViewDataPtr> VisibleRows;
 
@@ -224,6 +237,12 @@ protected:
 	/** The current filter text applied to the data table */
 	FText ActiveFilterText;
 
+	/** Currently selected sorting mode */
+	EColumnSortMode::Type SortMode;
+
+	/** Specify which column to sort with */
+	FName SortByColumn;
+
 	FOnRowHighlighted CallbackOnRowHighlighted;
 
 	FSimpleDelegate CallbackOnDataTableUndoRedo;
@@ -239,4 +258,7 @@ protected:
 
 	/** The column id for the row name list view column */
 	static const FName RowNameColumnId;
+
+	/** The column id for the row number list view column */
+	static const FName RowNumberColumnId;
 };
diff --git a/Engine/Source/Editor/DataTableEditor/Private/SDataTableListViewRow.cpp b/Engine/Source/Editor/DataTableEditor/Private/SDataTableListViewRow.cpp
--- a/Engine/Source/Editor/DataTableEditor/Private/SDataTableListViewRow.cpp
+++ b/Engine/Source/Editor/DataTableEditor/Private/SDataTableListViewRow.cpp
@@ -158,6 +158,19 @@ TSharedRef<SWidget> SDataTableListViewRow::MakeCellWidget(const int32 InRowIndex
 	FDataTableEditor* DataTableEdit = DataTableEditor.Pin().Get();
 	TArray<FDataTableEditorColumnHeaderDataPtr>& AvailableColumns = DataTableEdit->AvailableColumns;
 
+	if (InColumnId.IsEqual(FName(TEXT("RowNumber"))))
+	{
+		return SNew(SBox)
+			.Padding(FMargin(4, 2, 4, 2))
+			[
+				SNew(STextBlock)
+				.TextStyle(FEditorStyle::Get(), "DataTableEditor.CellText")
+				.Text(FText::FromString(FString::FromInt(RowDataPtr->RowNum)))
+				.ColorAndOpacity(DataTableEdit, &FDataTableEditor::GetRowTextColor, RowDataPtr->RowId)
+				.HighlightText(DataTableEdit, &FDataTableEditor::GetFilterText)
+			];
+	}
+
 	if (InColumnId.IsEqual(FName(TEXT("RowName"))))
 	{
 		return SNew(SBox)
@@ -166,6 +179,7 @@ TSharedRef<SWidget> SDataTableListViewRow::MakeCellWidget(const int32 InRowIndex
 				SAssignNew(InlineEditableText, SInlineEditableTextBlock)
 				.Text(RowDataPtr->DisplayName)
 				.OnTextCommitted(this, &SDataTableListViewRow::OnRowRenamed)
+				.HighlightText(DataTableEdit, &FDataTableEditor::GetFilterText)
 				.ColorAndOpacity(DataTableEdit, &FDataTableEditor::GetRowTextColor, RowDataPtr->RowId)
 			];
 	}
diff --git a/Engine/Source/Editor/DataTableEditor/Private/SRowEditor.cpp b/Engine/Source/Editor/DataTableEditor/Private/SRowEditor.cpp
--- a/Engine/Source/Editor/DataTableEditor/Private/SRowEditor.cpp
+++ b/Engine/Source/Editor/DataTableEditor/Private/SRowEditor.cpp
@@ -241,12 +241,6 @@ void SRowEditor::OnSelectionChanged(TSharedPtr<FName> InItem, ESelectInfo::Type
 		SelectedName = InItem;
 
 		Restore();
-
-		if (RenameTextBox.IsValid())
-		{
-			// refresh name, in case of a pending rename action
-			RenameTextBox->SetText(TAttribute<FText>(this, &SRowEditor::GetCurrentNameAsText));
-		}
 	}
 }
 
diff --git a/Engine/Source/Editor/DataTableEditor/Private/SRowEditor.h b/Engine/Source/Editor/DataTableEditor/Private/SRowEditor.h
--- a/Engine/Source/Editor/DataTableEditor/Private/SRowEditor.h
+++ b/Engine/Source/Editor/DataTableEditor/Private/SRowEditor.h
@@ -53,7 +53,6 @@ protected:
 	TSharedPtr<class IStructureDetailsView> StructureDetailsView;
 	TSharedPtr<FName> SelectedName;
 	TSharedPtr<SComboBox<TSharedPtr<FName>>> RowComboBox;
-	TSharedPtr<SEditableTextBox> RenameTextBox;
 
 	void RefreshNameList();
 	void CleanBeforeChange();
diff --git a/Engine/Source/Editor/UnrealEd/Private/DataTableEditorUtils.cpp b/Engine/Source/Editor/UnrealEd/Private/DataTableEditorUtils.cpp
--- a/Engine/Source/Editor/UnrealEd/Private/DataTableEditorUtils.cpp
+++ b/Engine/Source/Editor/UnrealEd/Private/DataTableEditorUtils.cpp
@@ -766,6 +766,7 @@ void FDataTableEditorUtils::CacheDataTableForEditing(const UDataTable* DataTable
 		}
 
 		CachedRowData->DesiredRowHeight = FontMeasure->GetMaxCharacterHeight(CellTextStyle.Font);
+		CachedRowData->RowNum = Index + 1;
 
 		// Always rebuild cell data
 		{
diff --git a/Engine/Source/Editor/UnrealEd/Public/DataTableEditorUtils.h b/Engine/Source/Editor/UnrealEd/Public/DataTableEditorUtils.h
--- a/Engine/Source/Editor/UnrealEd/Public/DataTableEditorUtils.h
+++ b/Engine/Source/Editor/UnrealEd/Public/DataTableEditorUtils.h
@@ -33,6 +33,9 @@ struct FDataTableEditorRowListViewData
 	/** The calculated height of this row taking into account the cell data for each column */
 	float DesiredRowHeight;
 
+	/** Insertion number of the row */
+	uint32 RowNum;
+
 	/** Array corresponding to each cell in this row */
 	TArray<FText> CellData;
 };
