Added link to open struct window on click, insertion of row above/below current row in context menu, drag drop row arrangement, and tool tips for index and drag drop columns all to Data Tables.

diff --git a/Engine/Source/Editor/DataTableEditor/Private/DataTableEditor.cpp b/Engine/Source/Editor/DataTableEditor/Private/DataTableEditor.cpp
--- a/Engine/Source/Editor/DataTableEditor/Private/DataTableEditor.cpp
+++ b/Engine/Source/Editor/DataTableEditor/Private/DataTableEditor.cpp
@@ -49,6 +49,7 @@ const FName FDataTableEditor::DataTableDetailsTabId("DataTableEditor_DataTableDe
 const FName FDataTableEditor::RowEditorTabId("DataTableEditor_RowEditor");
 const FName FDataTableEditor::RowNameColumnId("RowName");
 const FName FDataTableEditor::RowNumberColumnId("RowNumber");
+const FName FDataTableEditor::RowDragDropColumnId("RowDragDrop");
 
 void FDataTableEditor::RegisterTabSpawners(const TSharedRef<class FTabManager>& InTabManager)
 {
@@ -475,6 +476,22 @@ void FDataTableEditor::OnColumnNameSortModeChanged(const EColumnSortPriority::Ty
 	CellsListView->RequestListRefresh();
 }
 
+void FDataTableEditor::OnEditDataTableStructClicked()
+{
+	const UDataTable* DataTable = GetDataTable();
+	if (DataTable)
+	{
+
+		const UScriptStruct* ScriptStruct = DataTable->GetRowStruct();
+
+		if (ScriptStruct)
+		{
+			FAssetEditorManager::Get().OpenEditorForAsset(ScriptStruct->GetPathName());
+			FSourceCodeNavigation::NavigateToStruct(ScriptStruct);
+		}
+	}
+}
+
 UDataTable* FDataTableEditor::GetEditableDataTable() const
 {
 	return Cast<UDataTable>(GetEditingObject());
@@ -531,6 +548,31 @@ FText FDataTableEditor::GetCellToolTipText(FDataTableEditorRowListViewDataPtr In
 	return TooltipText;
 }
 
+float FDataTableEditor::GetRowNumberColumnWidth() const
+{
+	return RowNumberColumnWidth;
+}
+
+void FDataTableEditor::RefreshRowNumberColumnWidth()
+{
+
+	TSharedRef<FSlateFontMeasure> FontMeasure = FSlateApplication::Get().GetRenderer()->GetFontMeasureService();
+	const FTextBlockStyle& CellTextStyle = FEditorStyle::GetWidgetStyle<FTextBlockStyle>("DataTableEditor.CellText");
+	const float CellPadding = 10.0f;
+	RowNumberColumnWidth = 10.0f;
+	for (const FDataTableEditorRowListViewDataPtr& RowData : AvailableRows)
+	{
+		const float RowNumberWidth = FontMeasure->Measure(FString::FromInt(RowData->RowNum), CellTextStyle.Font).X + CellPadding;
+		RowNumberColumnWidth = FMath::Max(RowNumberColumnWidth, RowNumberWidth);
+	}
+
+}
+
+void FDataTableEditor::OnRowNumberColumnResized(const float NewWidth)
+{
+	RowNumberColumnWidth = NewWidth;
+}
+
 float FDataTableEditor::GetRowNameColumnWidth() const
 {
 	return RowNameColumnWidth;
@@ -837,11 +879,11 @@ void FDataTableEditor::PostRegenerateMenusAndToolbars()
 			.AutoWidth()
 			.VAlign(VAlign_Center)
 			[
-				SNew(STextBlock)
-				.ShadowOffset(FVector2D::UnitVector)
+				SNew(SHyperlink)
+				.Style(FEditorStyle::Get(), "Common.GotoNativeCodeHyperlink")
+				.OnNavigate(this, &FDataTableEditor::OnEditDataTableStructClicked)
 				.Text(FText::FromName(DataTable->GetRowStructName()))
-				.ToolTipText(LOCTEXT("DataTableRowToolTip", "The struct used for each row in this data table"))
-				.Visibility(UDS ? EVisibility::Visible : EVisibility::Collapsed)
+				.ToolTipText(LOCTEXT("DataTableRowToolTip", "Open the struct used for each row in this data table"))
 			]
 			+ SHorizontalBox::Slot()
 			.AutoWidth()
@@ -905,8 +947,9 @@ void FDataTableEditor::RefreshCachedDataTable(const FName InCachedSelection, con
 
 	FDataTableEditorUtils::CacheDataTableForEditing(Table, AvailableColumns, AvailableRows);
 
-	// Update the desired width of the row names column
+	// Update the desired width of the row names and numbers column
 	// This prevents it growing or shrinking as you scroll the list view
+	RefreshRowNumberColumnWidth();
 	RefreshRowNameColumnWidth();
 
 	// Setup the default auto-sized columns
@@ -942,16 +985,50 @@ void FDataTableEditor::RefreshCachedDataTable(const FName InCachedSelection, con
 	{
 		ColumnNamesHeaderRow->ClearColumns();
 
+		ColumnNamesHeaderRow->AddColumn(
+			SHeaderRow::Column(RowDragDropColumnId)
+			[
+				SNew(SBox)
+				.VAlign(VAlign_Fill) 
+				.HAlign(HAlign_Fill)
+				.ToolTip(IDocumentation::Get()->CreateToolTip(
+				LOCTEXT("DataTableRowHandleTooltip", "Drag Drop Handles"),
+				nullptr,
+				*FDataTableEditorUtils::VariableTypesTooltipDocLink,
+				TEXT("DataTableRowHandle")))
+				[
+					SNew(STextBlock)
+					.Text(FText::GetEmpty())
+				]
+			]
+		);
+
 		ColumnNamesHeaderRow->AddColumn(
 			SHeaderRow::Column(RowNumberColumnId)
-			.DefaultLabel(FText::GetEmpty())
 			.SortMode(this, &FDataTableEditor::GetColumnSortMode, RowNumberColumnId)
 			.OnSort(this, &FDataTableEditor::OnColumnNumberSortModeChanged)
+			.ManualWidth(this, &FDataTableEditor::GetRowNumberColumnWidth)
+			.OnWidthChanged(this, &FDataTableEditor::OnRowNumberColumnResized)
+			[
+				SNew(SBox)
+				.VAlign(VAlign_Fill)
+				.HAlign(HAlign_Fill)
+				.ToolTip(IDocumentation::Get()->CreateToolTip(
+				LOCTEXT("DataTableRowIndexTooltip", "Row Index"),
+				nullptr,
+				*FDataTableEditorUtils::VariableTypesTooltipDocLink,
+				TEXT("DataTableRowIndex")))
+				[
+					SNew(STextBlock)
+					.Text(FText::GetEmpty())
+				]
+			]
+
 		);
 
 		ColumnNamesHeaderRow->AddColumn(
 			SHeaderRow::Column(RowNameColumnId)
-			.DefaultLabel(FText::GetEmpty())
+			.DefaultLabel(LOCTEXT("DataTableRowName", "Row Name"))
 			.ManualWidth(this, &FDataTableEditor::GetRowNameColumnWidth)
 			.OnWidthChanged(this, &FDataTableEditor::OnRowNameColumnResized)
 			.SortMode(this, &FDataTableEditor::GetColumnSortMode, RowNameColumnId)
@@ -1251,6 +1328,13 @@ TSharedRef<SVerticalBox> FDataTableEditor::CreateContentBox()
 				]
 			]
 			+ SHorizontalBox::Slot()
+			.AutoWidth()
+			.Padding(2)
+			[
+				SNew(SSeparator)
+				.Orientation(Orient_Vertical)
+			]
+			+ SHorizontalBox::Slot()
 			[
 				SAssignNew(SearchBoxWidget, SSearchBox)
 				.InitialText(this, &FDataTableEditor::GetFilterText)
diff --git a/Engine/Source/Editor/DataTableEditor/Private/DataTableEditor.h b/Engine/Source/Editor/DataTableEditor/Private/DataTableEditor.h
--- a/Engine/Source/Editor/DataTableEditor/Private/DataTableEditor.h
+++ b/Engine/Source/Editor/DataTableEditor/Private/DataTableEditor.h
@@ -122,12 +122,17 @@ protected:
 	float GetRowNameColumnWidth() const;
 	void RefreshRowNameColumnWidth();
 
+	float GetRowNumberColumnWidth() const;
+	void RefreshRowNumberColumnWidth();
+
 	float GetColumnWidth(const int32 ColumnIndex) const;
 
 	void OnColumnResized(const float NewWidth, const int32 ColumnIndex);
 
 	void OnRowNameColumnResized(const float NewWidth);
 
+	void OnRowNumberColumnResized(const float NewWidth);
+
 	void LoadLayoutData();
 
 	void SaveLayoutData();
@@ -172,6 +177,8 @@ protected:
 	void OnColumnNumberSortModeChanged(const EColumnSortPriority::Type SortPriority, const FName& ColumnId, const EColumnSortMode::Type InSortMode);
 	void OnColumnNameSortModeChanged(const EColumnSortPriority::Type SortPriority, const FName& ColumnId, const EColumnSortMode::Type InSortMode);
 
+	void OnEditDataTableStructClicked();
+
 private:
 	UDataTable* GetEditableDataTable() const;
 
@@ -222,6 +229,9 @@ protected:
 	/** Width of the row name column */
 	float RowNameColumnWidth;
 
+	/** Width of the row number column */
+	float RowNumberColumnWidth;
+
 	/** Widths of data table cell columns */
 	TArray<FColumnWidth> ColumnWidths;
 
@@ -261,4 +271,7 @@ protected:
 
 	/** The column id for the row number list view column */
 	static const FName RowNumberColumnId;
+
+	/** The column id for the drag drop column */
+	static const FName RowDragDropColumnId;
 };
diff --git a/Engine/Source/Editor/DataTableEditor/Private/SDataTableListViewRow.cpp b/Engine/Source/Editor/DataTableEditor/Private/SDataTableListViewRow.cpp
--- a/Engine/Source/Editor/DataTableEditor/Private/SDataTableListViewRow.cpp
+++ b/Engine/Source/Editor/DataTableEditor/Private/SDataTableListViewRow.cpp
@@ -19,7 +19,8 @@ void SDataTableListViewRow::Construct(const FArguments& InArgs, const TSharedRef
 	DataTableEditor = InArgs._DataTableEditor;
 	SMultiColumnTableRow<FDataTableEditorRowListViewDataPtr>::Construct(
 		FSuperRowType::FArguments()
-		.Style(FEditorStyle::Get(), "DataTableEditor.CellListViewRow"),
+		.Style(FEditorStyle::Get(), "DataTableEditor.CellListViewRow")
+		.OnDrop(this, &SDataTableListViewRow::OnRowDrop),
 		InOwnerTableView
 	);
 
@@ -32,7 +33,10 @@ FReply SDataTableListViewRow::OnMouseButtonUp(const FGeometry& MyGeometry, const
 		FDataTableEditorUtils::SelectRow(DataTableEditor.Pin()->GetDataTable(), RowDataPtr->RowId);
 		TSharedRef<SWidget> MenuWidget = FDataTableRowUtils::MakeRowActionsMenu(DataTableEditor.Pin(),
 			FExecuteAction::CreateSP(this, &SDataTableListViewRow::OnSearchForReferences),
-			FExecuteAction::CreateSP(this, &SDataTableListViewRow::OnInsertNewRow));
+			FExecuteAction::CreateSP(this, &SDataTableListViewRow::OnInsertNewRow, ERowInsertionPosition::Bottom),
+			FExecuteAction::CreateSP(this, &SDataTableListViewRow::OnInsertNewRow, ERowInsertionPosition::Above),
+			FExecuteAction::CreateSP(this, &SDataTableListViewRow::OnInsertNewRow, ERowInsertionPosition::Below)
+		);
 
 		FWidgetPath WidgetPath = MouseEvent.GetEventPath() != nullptr ? *MouseEvent.GetEventPath() : FWidgetPath();
 		FSlateApplication::Get().PushMenu(AsShared(), WidgetPath, MenuWidget, MouseEvent.GetScreenSpacePosition(), FPopupTransitionEffect::ContextMenu);
@@ -58,7 +62,7 @@ void SDataTableListViewRow::OnSearchForReferences()
 	}
 }
 
-void SDataTableListViewRow::OnInsertNewRow()
+void SDataTableListViewRow::OnInsertNewRow(ERowInsertionPosition InsertPosition)
 {
 	if (DataTableEditor.IsValid() && RowDataPtr.IsValid())
 	{
@@ -74,13 +78,57 @@ void SDataTableListViewRow::OnInsertNewRow()
 					NewName.SetNumber(NewName.GetNumber() + 1);
 				}
 
-				FDataTableEditorUtils::AddRow(SourceDataTable, NewName);
+				if (InsertPosition == ERowInsertionPosition::Bottom)
+				{
+					FDataTableEditorUtils::AddRow(SourceDataTable, NewName);
+				}
+				else
+				{
+					FDataTableEditorUtils::AddRowAboveOrBelowSelection(SourceDataTable, *CurrentName, NewName, InsertPosition);
+				}
 				FDataTableEditorUtils::SelectRow(SourceDataTable, NewName);
 			}
 		}
 	}
 }
 
+FReply SDataTableListViewRow::OnRowDrop(const FDragDropEvent& DragDropEvent)
+{
+	TSharedPtr<FDataTableRowDragDropOp> DataTableDropOp = DragDropEvent.GetOperationAs< FDataTableRowDragDropOp >();
+	TSharedPtr<SDataTableListViewRow> RowPtr = nullptr;
+	if (DataTableDropOp.IsValid() && DataTableDropOp->Row.IsValid())
+	{
+		RowPtr = DataTableDropOp->Row.Pin();
+	}
+	if (!RowPtr.IsValid())
+	{
+		return FReply::Unhandled();
+	}
+
+	int32 JumpCount = (RowPtr->RowDataPtr)->RowNum - RowDataPtr->RowNum;
+
+	if (!JumpCount)
+	{
+		return FReply::Handled();
+	}
+
+	FDataTableEditorUtils::ERowMoveDirection Direction = JumpCount > 0 ? FDataTableEditorUtils::ERowMoveDirection::Up : FDataTableEditorUtils::ERowMoveDirection::Down;
+
+	if (FDataTableEditor* DataTableEditorPtr = DataTableEditor.Pin().Get())
+	{
+		UDataTable* SourceDataTable = const_cast<UDataTable*>(DataTableEditorPtr->GetDataTable());
+
+		if (SourceDataTable)
+		{
+			FDataTableEditorUtils::MoveRow(SourceDataTable, (RowPtr->RowDataPtr)->RowId, Direction, FMath::Abs<int32>(JumpCount));
+
+			return FReply::Handled();
+		}
+	}
+
+	return FReply::Unhandled();
+}
+
 FReply SDataTableListViewRow::OnKeyDown(const FGeometry& MyGeometry, const FKeyEvent& InKeyEvent)
 {
 	FKey Key = InKeyEvent.GetKey();
@@ -153,12 +201,32 @@ TSharedRef<SWidget> SDataTableListViewRow::GenerateWidgetForColumn(const FName&
 
 TSharedRef<SWidget> SDataTableListViewRow::MakeCellWidget(const int32 InRowIndex, const FName& InColumnId)
 {
+	const FName RowDragDropColumnId("RowDragDrop");
+
 	int32 ColumnIndex = 0;
 
 	FDataTableEditor* DataTableEdit = DataTableEditor.Pin().Get();
 	TArray<FDataTableEditorColumnHeaderDataPtr>& AvailableColumns = DataTableEdit->AvailableColumns;
 
-	if (InColumnId.IsEqual(FName(TEXT("RowNumber"))))
+	if (InColumnId.IsEqual(RowDragDropColumnId))
+	{
+		return SNew(SDataTableRowHandle)
+		.Content()
+			[
+				SNew(SHorizontalBox)
+				+ SHorizontalBox::Slot()
+				.Padding(5.0f, 1.0f)
+				[
+					SNew(SImage)
+					.Image(FCoreStyle::Get().GetBrush("VerticalBoxDragIndicatorShort"))
+				]
+			]
+		.ParentRow(SharedThis(this));
+	}
+
+	const FName RowNumberColumnId("RowNumber");
+
+	if (InColumnId.IsEqual(RowNumberColumnId))
 	{
 		return SNew(SBox)
 			.Padding(FMargin(4, 2, 4, 2))
@@ -171,7 +239,9 @@ TSharedRef<SWidget> SDataTableListViewRow::MakeCellWidget(const int32 InRowIndex
 			];
 	}
 
-	if (InColumnId.IsEqual(FName(TEXT("RowName"))))
+	const FName RowNameColumnId("RowName");
+
+	if (InColumnId.IsEqual(RowNameColumnId))
 	{
 		return SNew(SBox)
 			.Padding(FMargin(4, 2, 4, 2))
@@ -237,4 +307,58 @@ FText SDataTableListViewRow::GetCurrentNameAsText() const
 	return FText::FromName(GetCurrentName());
 }
 
+void SDataTableRowHandle::Construct(const FArguments& InArgs)
+{
+	ParentRow = InArgs._ParentRow;
+
+	ChildSlot
+		[
+			InArgs._Content.Widget
+		];
+}
+
+FReply SDataTableRowHandle::OnDragDetected(const FGeometry& MyGeometry, const FPointerEvent& MouseEvent)
+{
+	if (MouseEvent.IsMouseButtonDown(EKeys::LeftMouseButton))
+	{
+		TSharedPtr<FDragDropOperation> DragDropOp = CreateDragDropOperation(ParentRow.Pin());
+		if (DragDropOp.IsValid())
+		{
+			return FReply::Handled().BeginDragDrop(DragDropOp.ToSharedRef());
+		}
+	}
+
+	return FReply::Unhandled();
+
+}
+
+TSharedPtr<FDataTableRowDragDropOp> SDataTableRowHandle::CreateDragDropOperation(TSharedPtr<SDataTableListViewRow> InRow)
+{
+	TSharedPtr<FDataTableRowDragDropOp> Operation = MakeShareable(new FDataTableRowDragDropOp(InRow));
+
+	return Operation;
+}
+
+FDataTableRowDragDropOp::FDataTableRowDragDropOp(TSharedPtr<SDataTableListViewRow> InRow)
+{
+	Row = InRow;
+
+	DecoratorWidget = SNew(SBorder)
+		.Padding(8.f)
+		.BorderImage(FEditorStyle::GetBrush("Graph.ConnectorFeedback.Border"))
+		.Content()
+		[
+			SNew(SHorizontalBox)
+			+ SHorizontalBox::Slot()
+			.AutoWidth()
+			.VAlign(VAlign_Center)
+			[
+				SNew(STextBlock)
+				.Text(LOCTEXT("PlaceRowHere", "Place Row Here"))
+			]
+		];
+
+	Construct();
+}
+
 #undef LOCTEXT_NAMESPACE
diff --git a/Engine/Source/Editor/DataTableEditor/Private/SDataTableListViewRow.h b/Engine/Source/Editor/DataTableEditor/Private/SDataTableListViewRow.h
--- a/Engine/Source/Editor/DataTableEditor/Private/SDataTableListViewRow.h
+++ b/Engine/Source/Editor/DataTableEditor/Private/SDataTableListViewRow.h
@@ -10,9 +10,36 @@
 #include "Framework/Application/MenuStack.h"
 #include "Framework/Application/SlateApplication.h"
 #include "DataTableEditorUtils.h"
+#include "DragAndDrop/DecoratedDragDropOp.h"
 
 class FDataTableEditor;
 class SInlineEditableTextBlock;
+class SDataTableListViewRow;
+
+class SDataTableRowHandle: public SCompoundWidget
+{
+public:
+	SLATE_BEGIN_ARGS(SDataTableRowHandle)
+	{}
+	SLATE_DEFAULT_SLOT(FArguments, Content)
+		SLATE_ARGUMENT(TSharedPtr<SDataTableListViewRow>, ParentRow)
+		SLATE_END_ARGS()
+
+		void Construct(const FArguments& InArgs);
+
+	FReply OnMouseButtonDown(const FGeometry& MyGeometry, const FPointerEvent& MouseEvent) override
+	{
+		return FReply::Handled().DetectDrag(SharedThis(this), EKeys::LeftMouseButton);
+	};
+
+
+	FReply OnDragDetected(const FGeometry& MyGeometry, const FPointerEvent& MouseEvent) override;
+	TSharedPtr<class FDataTableRowDragDropOp> CreateDragDropOperation(TSharedPtr<SDataTableListViewRow> InRow);
+
+private:
+	TWeakPtr<SDataTableListViewRow> ParentRow;
+};
+
 /**
  * A widget to represent a row in a Data Table Editor widget. This widget allows us to do things like right-click
  * and take actions on a particular row of a Data Table.
@@ -48,7 +75,9 @@ public:
 private:
 
 	void OnSearchForReferences();
-	void OnInsertNewRow();
+	void OnInsertNewRow(ERowInsertionPosition InsertPosition);
+	
+	FReply OnRowDrop(const FDragDropEvent& DragDropEvent);
 
 	TSharedRef<SWidget> MakeCellWidget(const int32 InRowIndex, const FName& InColumnId);
 
@@ -58,4 +87,22 @@ private:
 
 	FDataTableEditorRowListViewDataPtr RowDataPtr;
 	TWeakPtr<FDataTableEditor> DataTableEditor;
-};
\ No newline at end of file
+
+};
+
+class FDataTableRowDragDropOp : public FDecoratedDragDropOp
+{
+public:
+	DRAG_DROP_OPERATOR_TYPE(FDataTableRowDragDropOp, FDecoratedDragDropOp)
+
+		FDataTableRowDragDropOp(TSharedPtr<SDataTableListViewRow> InRow);
+
+	TSharedPtr<SWidget> DecoratorWidget;
+
+	virtual TSharedPtr<SWidget> GetDefaultDecorator() const override
+	{
+		return DecoratorWidget;
+	}
+
+	TWeakPtr<class SDataTableListViewRow> Row;
+};
diff --git a/Engine/Source/Editor/DataTableEditor/Public/DataTableRowUtlis.cpp b/Engine/Source/Editor/DataTableEditor/Public/DataTableRowUtlis.cpp
--- a/Engine/Source/Editor/DataTableEditor/Public/DataTableRowUtlis.cpp
+++ b/Engine/Source/Editor/DataTableEditor/Public/DataTableRowUtlis.cpp
@@ -14,16 +14,29 @@
 const FText FDataTableRowUtils::SearchForReferencesActionName = LOCTEXT("FDataTableRowUtils_SearchForReferences", "Find Row References");
 const FText FDataTableRowUtils::SearchForReferencesActionTooltip = LOCTEXT("FDataTableRowUtils_SearchForReferencesTooltip", "Find assets that reference this Row");
 
-const FText InsertNewRowActionName = LOCTEXT("FDataTableRowUtils_InsertNewRow", "Insert New Row");
-const FText InsertNewRowActionTooltip = LOCTEXT("FDataTableRowUtils_InsertNewRowTooltip", "Insert a new Row");
+const FText FDataTableRowUtils::InsertNewRowActionName = LOCTEXT("FDataTableRowUtils_InsertNewRow", "Insert New Row");
+const FText FDataTableRowUtils::InsertNewRowActionTooltip = LOCTEXT("FDataTableRowUtils_InsertNewRowTooltip", "Insert a new Row");
 
-TSharedRef<SWidget> FDataTableRowUtils::MakeRowActionsMenu(TSharedPtr<IDataTableEditor> Editor, FExecuteAction SearchForReferencesAction, FExecuteAction InsertNewRowAction)
+const FText FDataTableRowUtils::InsertNewRowAboveActionName = LOCTEXT("FDataTableRowUtils_InsertNewRowAbove", "Insert New Row Above");
+const FText FDataTableRowUtils::InsertNewRowAboveActionTooltip = LOCTEXT("FDataTableRowUtils_InsertNewRowAboveTooltip", "Insert a new Row above the current selection");
+
+const FText FDataTableRowUtils::InsertNewRowBelowActionName = LOCTEXT("FDataTableRowUtils_InsertNewRowBelow", "Insert New Row Below");
+const FText FDataTableRowUtils::InsertNewRowBelowActionTooltip = LOCTEXT("FDataTableRowUtils_InsertNewRowBelowTooltip", "Insert a new Row below the current selection");
+
+
+TSharedRef<SWidget> FDataTableRowUtils::MakeRowActionsMenu(TSharedPtr<IDataTableEditor> Editor, FExecuteAction SearchForReferencesAction, FExecuteAction InsertNewRowAction,
+	FExecuteAction InsertNewRowAboveAction, FExecuteAction InsertNewRowBelowAction)
 {
-	if (SearchForReferencesAction.IsBound())
+	if (SearchForReferencesAction.IsBound() && InsertNewRowAction.IsBound() && InsertNewRowAboveAction.IsBound() && InsertNewRowBelowAction.IsBound() )
 	{
 		FMenuBuilder MenuBuilder(true, Editor->GetToolkitCommands());
 		MenuBuilder.AddMenuEntry(InsertNewRowActionName, InsertNewRowActionTooltip,
 			FSlateIcon(FEditorStyle::GetStyleSetName(), "Plus"), FUIAction(InsertNewRowAction));
+		MenuBuilder.AddMenuEntry(InsertNewRowAboveActionName, InsertNewRowAboveActionTooltip,
+			FSlateIcon(), FUIAction(InsertNewRowAboveAction));
+		MenuBuilder.AddMenuEntry(InsertNewRowBelowActionName, InsertNewRowBelowActionTooltip,
+			FSlateIcon(), FUIAction(InsertNewRowBelowAction));
+
 		MenuBuilder.AddMenuEntry(FGenericCommands::Get().Copy);
 		MenuBuilder.AddMenuEntry(FGenericCommands::Get().Paste);
 		MenuBuilder.AddMenuEntry(FGenericCommands::Get().Duplicate);
diff --git a/Engine/Source/Editor/DataTableEditor/Public/DataTableRowUtlis.h b/Engine/Source/Editor/DataTableEditor/Public/DataTableRowUtlis.h
--- a/Engine/Source/Editor/DataTableEditor/Public/DataTableRowUtlis.h
+++ b/Engine/Source/Editor/DataTableEditor/Public/DataTableRowUtlis.h
@@ -12,11 +12,23 @@ class FDetailWidgetRow;
 class DATATABLEEDITOR_API FDataTableRowUtils
 {
 public:
-	static TSharedRef<SWidget> MakeRowActionsMenu(TSharedPtr<class IDataTableEditor> Editor, FExecuteAction SearchForReferencesAction, FExecuteAction InsertNewRowAction);
+	static TSharedRef<SWidget> MakeRowActionsMenu(TSharedPtr<class IDataTableEditor> Editor, FExecuteAction SearchForReferencesAction, FExecuteAction InsertNewRowAction,
+													FExecuteAction InsertNewRowAboveAction, FExecuteAction InsertNewRowBelowAction);
 	static void AddSearchForReferencesContextMenu(FDetailWidgetRow& RowNameDetailWidget, FExecuteAction SearchForReferencesAction);
 
 private:
 	static const FText SearchForReferencesActionName;
 	static const FText SearchForReferencesActionTooltip;
+
+	static const FText InsertNewRowActionName;
+	static const FText InsertNewRowActionTooltip;
+
+	static const FText InsertNewRowAboveActionName;
+	static const FText InsertNewRowAboveActionTooltip;
+
+	static const FText InsertNewRowBelowActionName;
+	static const FText InsertNewRowBelowActionTooltip;
+
+
 };
 
diff --git a/Engine/Source/Editor/UnrealEd/Private/DataTableEditorUtils.cpp b/Engine/Source/Editor/UnrealEd/Private/DataTableEditorUtils.cpp
--- a/Engine/Source/Editor/UnrealEd/Private/DataTableEditorUtils.cpp
+++ b/Engine/Source/Editor/UnrealEd/Private/DataTableEditorUtils.cpp
@@ -493,6 +493,66 @@ uint8* FDataTableEditorUtils::AddRow(UDataTable* DataTable, FName RowName)
 	return RowData;
 }
 
+uint8* FDataTableEditorUtils::AddRowAboveOrBelowSelection(UDataTable* DataTable, const FName& RowName, const FName& NewRowName, ERowInsertionPosition InsertPosition)
+{
+	if (!DataTable || (NewRowName == NAME_None) || (DataTable->GetRowMap().Find(NewRowName) != nullptr) || !DataTable->RowStruct)
+	{
+		return nullptr;
+	}
+
+	const FScopedTransaction Transaction(LOCTEXT("AddDataTableRowAboveBelow", "Add Data Table Row Above or Below"));
+
+	TArray<FName> OrderedRowNames;
+	DataTable->GetRowMap().GenerateKeyArray(OrderedRowNames);
+
+	int32 CurrentRowIndex = OrderedRowNames.IndexOfByKey(RowName);
+	if (CurrentRowIndex == INDEX_NONE)
+	{
+		return nullptr;
+	}
+
+	if (InsertPosition == ERowInsertionPosition::Below)
+	{
+		CurrentRowIndex += 1;
+	}
+
+	OrderedRowNames.Insert(NewRowName, CurrentRowIndex);
+	
+	// Build a name -> index map as the KeySort will hit this a lot
+	TMap<FName, int32> NamesToNewIndex;
+	for (int32 NameIndex = 0; NameIndex < OrderedRowNames.Num(); ++NameIndex)
+	{
+		NamesToNewIndex.Add(OrderedRowNames[NameIndex], NameIndex);
+	}
+	
+	
+	BroadcastPreChange(DataTable, EDataTableChangeInfo::RowList);
+	
+	DataTable->Modify();
+	
+	// Allocate data to store information, using UScriptStruct to know its size
+	uint8* RowData = (uint8*)FMemory::Malloc(DataTable->RowStruct->GetStructureSize());
+
+	// And be sure to call DestroyScriptStruct later
+	DataTable->RowStruct->InitializeStruct(RowData);
+
+	// Add to row map
+	DataTable->AddRowInternal(NewRowName, RowData);
+
+	// Re-sort the map keys to match the new order
+	DataTable->GetNonConstRowMap().KeySort([&NamesToNewIndex](const FName& One, const FName& Two) -> bool
+	{
+		const int32 OneIndex = NamesToNewIndex.FindRef(One);
+		const int32 TwoIndex = NamesToNewIndex.FindRef(Two);
+		return OneIndex < TwoIndex;
+	});
+
+	BroadcastPostChange(DataTable, EDataTableChangeInfo::RowList);
+
+	return RowData;
+}
+
+
 uint8* FDataTableEditorUtils::DuplicateRow(UDataTable* DataTable, FName SourceRowName, FName RowName)
 {
 	if (!DataTable || (SourceRowName == NAME_None) || !DataTable->RowMap.Contains(SourceRowName) || DataTable->RowMap.Contains(RowName) || !DataTable->RowStruct)
diff --git a/Engine/Source/Editor/UnrealEd/Public/DataTableEditorUtils.h b/Engine/Source/Editor/UnrealEd/Public/DataTableEditorUtils.h
--- a/Engine/Source/Editor/UnrealEd/Public/DataTableEditorUtils.h
+++ b/Engine/Source/Editor/UnrealEd/Public/DataTableEditorUtils.h
@@ -43,6 +43,14 @@ struct FDataTableEditorRowListViewData
 typedef TSharedPtr<FDataTableEditorColumnHeaderData> FDataTableEditorColumnHeaderDataPtr;
 typedef TSharedPtr<FDataTableEditorRowListViewData>  FDataTableEditorRowListViewDataPtr;
 
+
+enum class ERowInsertionPosition
+{
+	Above,
+	Below,
+	Bottom,
+};
+
 struct UNREALED_API FDataTableEditorUtils
 {
 	enum class EDataTableChangeInfo
@@ -83,6 +91,8 @@ struct UNREALED_API FDataTableEditorUtils
 	static bool DiffersFromDefault(UDataTable* DataTable, FName RowName);
 	static bool ResetToDefault(UDataTable* DataTable, FName RowName);
 
+	static uint8* AddRowAboveOrBelowSelection(UDataTable* DataTable, const FName& RowName, const FName& NewRowName, ERowInsertionPosition InsertPosition);
+
 	static void BroadcastPreChange(UDataTable* DataTable, EDataTableChangeInfo Info);
 	static void BroadcastPostChange(UDataTable* DataTable, EDataTableChangeInfo Info);
 
