Fix CIS build missing header

diff --git a/Engine/Source/Editor/DataTableEditor/Private/DataTableEditor.cpp b/Engine/Source/Editor/DataTableEditor/Private/DataTableEditor.cpp
--- a/Engine/Source/Editor/DataTableEditor/Private/DataTableEditor.cpp
+++ b/Engine/Source/Editor/DataTableEditor/Private/DataTableEditor.cpp
@@ -37,6 +37,7 @@
 #include "PropertyEditorModule.h"
 #include "UObject/StructOnScope.h"
 #include "Toolkits/GlobalEditorCommonCommands.h"
+#include "Toolkits/AssetEditorManager.h"
 #include "DataTableRowUtlis.h"
 #include "Engine/DataTable.h"
 
