Took out toolbar, but moved the save and browse onto the data table tool bar. Also added some row arraging functionality to context menu. Took out the icons for row arrangement. Increased the size of the icons. Added in drag drop row highlight. Made the current drag drop decorator use the row name.

diff --git a/Engine/Source/Editor/DataTableEditor/Private/DataTableEditor.cpp b/Engine/Source/Editor/DataTableEditor/Private/DataTableEditor.cpp
--- a/Engine/Source/Editor/DataTableEditor/Private/DataTableEditor.cpp
+++ b/Engine/Source/Editor/DataTableEditor/Private/DataTableEditor.cpp
@@ -56,8 +56,6 @@ void FDataTableEditor::RegisterTabSpawners(const TSharedRef<class FTabManager>&
 {
 	WorkspaceMenuCategory = InTabManager->AddLocalWorkspaceMenuCategory(LOCTEXT("WorkspaceMenu_Data Table Editor", "Data Table Editor"));
 
-	FAssetEditorToolkit::RegisterTabSpawners(InTabManager);
-
 	CreateAndRegisterDataTableTab(InTabManager);
 	CreateAndRegisterDataTableDetailsTab(InTabManager);
 	CreateAndRegisterRowEditorTab(InTabManager);
@@ -65,8 +63,6 @@ void FDataTableEditor::RegisterTabSpawners(const TSharedRef<class FTabManager>&
 
 void FDataTableEditor::UnregisterTabSpawners(const TSharedRef<class FTabManager>& InTabManager)
 {
-	FAssetEditorToolkit::UnregisterTabSpawners(InTabManager);
-
 	InTabManager->UnregisterTabSpawner(DataTableTabId);
 	InTabManager->UnregisterTabSpawner(DataTableDetailsTabId);
 	InTabManager->UnregisterTabSpawner(RowEditorTabId);
@@ -198,18 +194,11 @@ void FDataTableEditor::HandlePostChange()
 
 void FDataTableEditor::InitDataTableEditor( const EToolkitMode::Type Mode, const TSharedPtr< class IToolkitHost >& InitToolkitHost, UDataTable* Table )
 {
-	TSharedRef<FTabManager::FLayout> StandaloneDefaultLayout = FTabManager::NewLayout( "Standalone_DataTableEditor_Layout_v3" )
+	TSharedRef<FTabManager::FLayout> StandaloneDefaultLayout = FTabManager::NewLayout( "Standalone_DataTableEditor_Layout_v4" )
 	->AddArea
 	(
 		FTabManager::NewPrimaryArea()->SetOrientation(Orient_Vertical)
 		->Split
-		(
-			FTabManager::NewStack()
-			->SetSizeCoefficient(0.1f)
-			->SetHideTabWell(true)
-			->AddTab(GetToolbarTabId(), ETabState::OpenedTab)
-		)
-		->Split
 		(
 			FTabManager::NewStack()
 			->AddTab(DataTableTabId, ETabState::OpenedTab)
@@ -224,9 +213,9 @@ void FDataTableEditor::InitDataTableEditor( const EToolkitMode::Type Mode, const
 	);
 
 	const bool bCreateDefaultStandaloneMenu = true;
-	const bool bCreateDefaultToolbar = true;
+	const bool bCreateDefaultToolbar = false;
 	FAssetEditorToolkit::InitAssetEditor( Mode, InitToolkitHost, FDataTableEditorModule::DataTableEditorAppIdentifier, StandaloneDefaultLayout, bCreateDefaultStandaloneMenu, bCreateDefaultToolbar, Table );
-	
+
 	FDataTableEditorModule& DataTableEditorModule = FModuleManager::LoadModuleChecked<FDataTableEditorModule>( "DataTableEditor" );
 	AddMenuExtender(DataTableEditorModule.GetMenuExtensibilityManager()->GetAllExtenders(GetToolkitCommands(), GetEditingObjects()));
 
@@ -493,6 +482,20 @@ void FDataTableEditor::OnEditDataTableStructClicked()
 	}
 }
 
+FReply FDataTableEditor::SaveDataTable_Execute()
+{
+	FAssetEditorToolkit::SaveAsset_Execute();
+
+	return FReply::Handled();
+}
+
+FReply FDataTableEditor::BrowseForDataTable_Execute()
+{
+	FAssetEditorToolkit::FindInContentBrowser_Execute();
+
+	return FReply::Handled();
+}
+
 UDataTable* FDataTableEditor::GetEditableDataTable() const
 {
 	return Cast<UDataTable>(GetEditingObject());
@@ -894,7 +897,7 @@ void FDataTableEditor::PostRegenerateMenusAndToolbars()
 				.ButtonStyle(FEditorStyle::Get(), "HoverHintOnly")
 				.OnClicked(this, &FDataTableEditor::OnFindRowInContentBrowserClicked)
 				.Visibility(UDS ? EVisibility::Visible : EVisibility::Collapsed)
-				.ToolTipText(LOCTEXT("FindRowInCBToolTip", "Find row in Content Browser"))
+				.ToolTipText(LOCTEXT("FindRowInCBToolTip", "Find struct in Content Browser"))
 				.ContentPadding(4.0f)
 				.ForegroundColor(FSlateColor::UseForeground())
 				[
@@ -1182,11 +1185,11 @@ TSharedRef<SVerticalBox> FDataTableEditor::CreateContentBox()
 				.ForegroundColor(FSlateColor::UseForeground())
 				.HAlign(HAlign_Center)
 				.VAlign(VAlign_Center)
-				.OnClicked(this, &FDataTableEditor::OnAddClicked)
-				.ToolTipText(LOCTEXT("AddRowTooltip", "Add a new row to the data table"))
+				.OnClicked(this, &FDataTableEditor::SaveDataTable_Execute)
+				.ToolTipText(LOCTEXT("SaveDataTableTooltip", "Save data table"))
 				[
 					SNew(SImage)
-					.Image(FEditorStyle::Get().GetBrush("Plus"))
+					.Image(FEditorStyle::Get().GetBrush("AssetEditor.SaveAsset.Small"))
 				]
 			]
 			+ SHorizontalBox::Slot()
@@ -1198,11 +1201,11 @@ TSharedRef<SVerticalBox> FDataTableEditor::CreateContentBox()
 				.ForegroundColor(FSlateColor::UseForeground())
 				.HAlign(HAlign_Center)
 				.VAlign(VAlign_Center)
-				.OnClicked(this, &FDataTableEditor::OnCopyClicked)
-				.ToolTipText(LOCTEXT("CopyTooltip", "Copy the currently selected row"))
+				.OnClicked(this, &FDataTableEditor::BrowseForDataTable_Execute)
+				.ToolTipText(LOCTEXT("FindInContentBrowserTooltip", "Find data table in content browser"))
 				[
 					SNew(SImage)
-					.Image(FEditorStyle::Get().GetBrush("DataTableEditor.Copy"))
+					.Image(FEditorStyle::Get().GetBrush("Sequencer.FindInContentBrowser.Small"))
 				]
 			]
 			+ SHorizontalBox::Slot()
@@ -1214,11 +1217,11 @@ TSharedRef<SVerticalBox> FDataTableEditor::CreateContentBox()
 				.ForegroundColor(FSlateColor::UseForeground())
 				.HAlign(HAlign_Center)
 				.VAlign(VAlign_Center)
-				.OnClicked(this, &FDataTableEditor::OnPasteClicked)
-				.ToolTipText(LOCTEXT("PasteTooltip", "Paste on the currently selected row"))
+				.OnClicked(this, &FDataTableEditor::OnAddClicked)
+				.ToolTipText(LOCTEXT("AddRowTooltip", "Add a new row to the data table"))
 				[
 					SNew(SImage)
-					.Image(FEditorStyle::Get().GetBrush("DataTableEditor.Paste"))
+					.Image(FEditorStyle::Get().GetBrush("DataTableEditor.Add"))
 				]
 			]
 			+ SHorizontalBox::Slot()
@@ -1230,11 +1233,11 @@ TSharedRef<SVerticalBox> FDataTableEditor::CreateContentBox()
 				.ForegroundColor(FSlateColor::UseForeground())
 				.HAlign(HAlign_Center)
 				.VAlign(VAlign_Center)
-				.OnClicked(this, &FDataTableEditor::OnDuplicateClicked)
-				.ToolTipText(LOCTEXT("DuplicateTooltip", "Duplicate the currently selected row"))
+				.OnClicked(this, &FDataTableEditor::OnCopyClicked)
+				.ToolTipText(LOCTEXT("CopyTooltip", "Copy the currently selected row"))
 				[
 					SNew(SImage)
-					.Image(FEditorStyle::Get().GetBrush("DataTableEditor.Duplicate"))
+					.Image(FEditorStyle::Get().GetBrush("DataTableEditor.Copy"))
 				]
 			]
 			+ SHorizontalBox::Slot()
@@ -1246,52 +1249,11 @@ TSharedRef<SVerticalBox> FDataTableEditor::CreateContentBox()
 				.ForegroundColor(FSlateColor::UseForeground())
 				.HAlign(HAlign_Center)
 				.VAlign(VAlign_Center)
-				.OnClicked(this, &FDataTableEditor::OnRemoveClicked)
-				.ToolTipText(LOCTEXT("RemoveRowTooltip", "Remove the currently selected row from the data table"))
+				.OnClicked(this, &FDataTableEditor::OnPasteClicked)
+				.ToolTipText(LOCTEXT("PasteTooltip", "Paste on the currently selected row"))
 				[
 					SNew(SImage)
-					.Image(FEditorStyle::Get().GetBrush("Cross"))
-				]
-			]
-			+ SHorizontalBox::Slot()
-			.AutoWidth()
-			.Padding(2)
-			[
-				SNew(SSeparator)
-				.Orientation(Orient_Vertical)
-			]
-			+ SHorizontalBox::Slot()
-			.AutoWidth()
-			.Padding(2)
-			[
-				SNew(SButton)
-				.ButtonStyle(FEditorStyle::Get(), "HoverHintOnly")
-				.ForegroundColor(FSlateColor::UseForeground())
-				.HAlign(HAlign_Center)
-				.VAlign(VAlign_Center)
-				.OnClicked(this, &FDataTableEditor::OnMoveRowClicked, FDataTableEditorUtils::ERowMoveDirection::Up)
-				.ToolTipText(LOCTEXT("MoveUpTooltip", "Move the currently selected row up by one in the data table sorted by row number"))
-				[
-					SNew(STextBlock)
-					.Font(FEditorStyle::Get().GetFontStyle("FontAwesome.14"))
-					.Text(FText::FromString(FString(TEXT("\xf106"))) /*fa-angle-up*/)
-				]
-			]
-			+ SHorizontalBox::Slot()
-			.AutoWidth()
-			.Padding(2)
-			[
-				SNew(SButton)
-				.ButtonStyle(FEditorStyle::Get(), "HoverHintOnly")
-				.ForegroundColor(FSlateColor::UseForeground())
-				.HAlign(HAlign_Center)
-				.VAlign(VAlign_Center)
-				.OnClicked(this, &FDataTableEditor::OnMoveRowClicked, FDataTableEditorUtils::ERowMoveDirection::Down)
-				.ToolTipText(LOCTEXT("MoveDownTooltip", "Move the currently selected row down by one in the data table sorted by row number"))
-				[
-					SNew(STextBlock)
-					.Font(FEditorStyle::Get().GetFontStyle("FontAwesome.14"))
-					.Text(FText::FromString(FString(TEXT("\xf107"))) /*fa-angle-down*/)
+					.Image(FEditorStyle::Get().GetBrush("DataTableEditor.Paste"))
 				]
 			]
 			+ SHorizontalBox::Slot()
@@ -1303,12 +1265,11 @@ TSharedRef<SVerticalBox> FDataTableEditor::CreateContentBox()
 				.ForegroundColor(FSlateColor::UseForeground())
 				.HAlign(HAlign_Center)
 				.VAlign(VAlign_Center)
-				.OnClicked(this, &FDataTableEditor::OnMoveToExtentClicked, FDataTableEditorUtils::ERowMoveDirection::Up)
-				.ToolTipText(LOCTEXT("MoveToTopTooltip", "Move the currently selected row to the top of the data table sorted by row number"))
+				.OnClicked(this, &FDataTableEditor::OnDuplicateClicked)
+				.ToolTipText(LOCTEXT("DuplicateTooltip", "Duplicate the currently selected row"))
 				[
-					SNew(STextBlock)
-					.Font(FEditorStyle::Get().GetFontStyle("FontAwesome.14"))
-					.Text(FText::FromString(FString(TEXT("\xf102"))) /*fa-angle-double-up*/)
+					SNew(SImage)
+					.Image(FEditorStyle::Get().GetBrush("DataTableEditor.Duplicate"))
 				]
 			]
 			+ SHorizontalBox::Slot()
@@ -1320,12 +1281,11 @@ TSharedRef<SVerticalBox> FDataTableEditor::CreateContentBox()
 				.ForegroundColor(FSlateColor::UseForeground())
 				.HAlign(HAlign_Center)
 				.VAlign(VAlign_Center)
-				.OnClicked(this, &FDataTableEditor::OnMoveToExtentClicked, FDataTableEditorUtils::ERowMoveDirection::Down)
-				.ToolTipText(LOCTEXT("MoveToBottomTooltip", "Move the currently selected row to the bottom of the data table based on row number"))
+				.OnClicked(this, &FDataTableEditor::OnRemoveClicked)
+				.ToolTipText(LOCTEXT("RemoveRowTooltip", "Remove the currently selected row from the data table"))
 				[
-					SNew(STextBlock)
-					.Font(FEditorStyle::Get().GetFontStyle("FontAwesome.14"))
-					.Text(FText::FromString(FString(TEXT("\xf103"))) /*fa-angle-double-down*/)
+					SNew(SImage)
+					.Image(FEditorStyle::Get().GetBrush("DataTableEditor.Remove"))
 				]
 			]
 			+ SHorizontalBox::Slot()
diff --git a/Engine/Source/Editor/DataTableEditor/Private/DataTableEditor.h b/Engine/Source/Editor/DataTableEditor/Private/DataTableEditor.h
--- a/Engine/Source/Editor/DataTableEditor/Private/DataTableEditor.h
+++ b/Engine/Source/Editor/DataTableEditor/Private/DataTableEditor.h
@@ -179,6 +179,10 @@ protected:
 
 	void OnEditDataTableStructClicked();
 
+	FReply SaveDataTable_Execute();
+	FReply BrowseForDataTable_Execute();
+
+
 private:
 	UDataTable* GetEditableDataTable() const;
 
diff --git a/Engine/Source/Editor/DataTableEditor/Private/SDataTableListViewRow.cpp b/Engine/Source/Editor/DataTableEditor/Private/SDataTableListViewRow.cpp
--- a/Engine/Source/Editor/DataTableEditor/Private/SDataTableListViewRow.cpp
+++ b/Engine/Source/Editor/DataTableEditor/Private/SDataTableListViewRow.cpp
@@ -20,10 +20,14 @@ void SDataTableListViewRow::Construct(const FArguments& InArgs, const TSharedRef
 	SMultiColumnTableRow<FDataTableEditorRowListViewDataPtr>::Construct(
 		FSuperRowType::FArguments()
 		.Style(FEditorStyle::Get(), "DataTableEditor.CellListViewRow")
-		.OnDrop(this, &SDataTableListViewRow::OnRowDrop),
+		.OnDrop(this, &SDataTableListViewRow::OnRowDrop)
+		.OnDragEnter(this, &SDataTableListViewRow::OnRowDragEnter)
+		.OnDragLeave(this, &SDataTableListViewRow::OnRowDragLeave),
 		InOwnerTableView
 	);
 
+	BorderImage = TAttribute<const FSlateBrush*>(this, &SDataTableListViewRow::GetBorder);
+
 }
 
 FReply SDataTableListViewRow::OnMouseButtonUp(const FGeometry& MyGeometry, const FPointerEvent& MouseEvent)
@@ -35,7 +39,9 @@ FReply SDataTableListViewRow::OnMouseButtonUp(const FGeometry& MyGeometry, const
 			FExecuteAction::CreateSP(this, &SDataTableListViewRow::OnSearchForReferences),
 			FExecuteAction::CreateSP(this, &SDataTableListViewRow::OnInsertNewRow, ERowInsertionPosition::Bottom),
 			FExecuteAction::CreateSP(this, &SDataTableListViewRow::OnInsertNewRow, ERowInsertionPosition::Above),
-			FExecuteAction::CreateSP(this, &SDataTableListViewRow::OnInsertNewRow, ERowInsertionPosition::Below)
+			FExecuteAction::CreateSP(this, &SDataTableListViewRow::OnInsertNewRow, ERowInsertionPosition::Below),
+			FExecuteAction::CreateSP(this, &SDataTableListViewRow::OnMoveToExtentClicked, FDataTableEditorUtils::ERowMoveDirection::Down),
+			FExecuteAction::CreateSP(this, &SDataTableListViewRow::OnMoveToExtentClicked, FDataTableEditorUtils::ERowMoveDirection::Up)
 		);
 
 		FWidgetPath WidgetPath = MouseEvent.GetEventPath() != nullptr ? *MouseEvent.GetEventPath() : FWidgetPath();
@@ -120,7 +126,14 @@ FReply SDataTableListViewRow::OnRowDrop(const FDragDropEvent& DragDropEvent)
 
 		if (SourceDataTable)
 		{
-			FDataTableEditorUtils::MoveRow(SourceDataTable, (RowPtr->RowDataPtr)->RowId, Direction, FMath::Abs<int32>(JumpCount));
+			FName& RowId = (RowPtr->RowDataPtr)->RowId;
+
+			FDataTableEditorUtils::MoveRow(SourceDataTable, RowId, Direction, FMath::Abs<int32>(JumpCount));
+			
+			FDataTableEditorUtils::SelectRow(SourceDataTable, RowId);
+
+			DataTableEditorPtr->SortMode = EColumnSortMode::Ascending;
+			DataTableEditorPtr->SortByColumn = FDataTableEditor::RowNumberColumnId;
 
 			return FReply::Handled();
 		}
@@ -211,7 +224,7 @@ TSharedRef<SWidget> SDataTableListViewRow::MakeCellWidget(const int32 InRowIndex
 	if (InColumnId.IsEqual(RowDragDropColumnId))
 	{
 		return SNew(SDataTableRowHandle)
-		.Content()
+			.Content()
 			[
 				SNew(SHorizontalBox)
 				+ SHorizontalBox::Slot()
@@ -221,7 +234,7 @@ TSharedRef<SWidget> SDataTableListViewRow::MakeCellWidget(const int32 InRowIndex
 					.Image(FCoreStyle::Get().GetBrush("VerticalBoxDragIndicatorShort"))
 				]
 			]
-		.ParentRow(SharedThis(this));
+			.ParentRow(SharedThis(this));
 	}
 
 	const FName RowNumberColumnId("RowNumber");
@@ -302,6 +315,11 @@ void SDataTableListViewRow::SetRowForRename()
 	InlineEditableText->EnterEditingMode();
 }
 
+const FDataTableEditorRowListViewDataPtr& SDataTableListViewRow::GetRowDataPtr() const
+{
+	return RowDataPtr;
+}
+
 FText SDataTableListViewRow::GetCurrentNameAsText() const
 {
 	return FText::FromName(GetCurrentName());
@@ -339,10 +357,59 @@ TSharedPtr<FDataTableRowDragDropOp> SDataTableRowHandle::CreateDragDropOperation
 	return Operation;
 }
 
+void SDataTableListViewRow::SetIsDragDrop(bool bInIsDragDrop)
+{
+	bIsDragDropObject = bInIsDragDrop;
+}
+
+
+void SDataTableListViewRow::OnRowDragEnter(const FDragDropEvent& DragDropEvent)
+{
+	bIsHoveredDragTarget = true;
+}
+
+void SDataTableListViewRow::OnRowDragLeave(const FDragDropEvent& DragDropEvent)
+{
+	bIsHoveredDragTarget = false;
+}
+
+
+const FSlateBrush* SDataTableListViewRow::GetBorder() const
+{
+	if (bIsDragDropObject)
+	{
+		return FEditorStyle::GetBrush("DataTableEditor.DragDropObject");
+	}
+	else if (bIsHoveredDragTarget)
+	{
+		return FEditorStyle::GetBrush("DataTableEditor.DragDropHoveredTarget");
+	}
+	else
+	{
+		return STableRow::GetBorder();
+	}
+}
+
+void SDataTableListViewRow::OnMoveToExtentClicked(FDataTableEditorUtils::ERowMoveDirection MoveDirection)
+{
+	if (DataTableEditor.IsValid())
+	{
+		TSharedPtr<FDataTableEditor> DataTableEditorPtr = DataTableEditor.Pin();
+		DataTableEditorPtr->OnMoveToExtentClicked(MoveDirection);
+	}
+}
+
 FDataTableRowDragDropOp::FDataTableRowDragDropOp(TSharedPtr<SDataTableListViewRow> InRow)
 {
 	Row = InRow;
 
+	TSharedPtr<SDataTableListViewRow> RowPtr = nullptr;
+	if (Row.IsValid())
+	{
+		RowPtr = Row.Pin();
+		RowPtr->SetIsDragDrop(true);
+	}
+
 	DecoratorWidget = SNew(SBorder)
 		.Padding(8.f)
 		.BorderImage(FEditorStyle::GetBrush("Graph.ConnectorFeedback.Border"))
@@ -354,11 +421,23 @@ FDataTableRowDragDropOp::FDataTableRowDragDropOp(TSharedPtr<SDataTableListViewRo
 			.VAlign(VAlign_Center)
 			[
 				SNew(STextBlock)
-				.Text(LOCTEXT("PlaceRowHere", "Place Row Here"))
+				.Text(FText::Format(LOCTEXT("DragDropDecoratorText", "Place {0} Here"), FText::FromName(InRow->GetCurrentName())))
 			]
 		];
 
 	Construct();
 }
 
+void FDataTableRowDragDropOp::OnDrop(bool bDropWasHandled, const FPointerEvent& MouseEvent)
+{
+	FDecoratedDragDropOp::OnDrop(bDropWasHandled, MouseEvent);
+
+	TSharedPtr<SDataTableListViewRow> RowPtr = nullptr;
+	if (Row.IsValid())
+	{
+		RowPtr = Row.Pin();
+		RowPtr->SetIsDragDrop(false);
+	}
+}
+
 #undef LOCTEXT_NAMESPACE
diff --git a/Engine/Source/Editor/DataTableEditor/Private/SDataTableListViewRow.h b/Engine/Source/Editor/DataTableEditor/Private/SDataTableListViewRow.h
--- a/Engine/Source/Editor/DataTableEditor/Private/SDataTableListViewRow.h
+++ b/Engine/Source/Editor/DataTableEditor/Private/SDataTableListViewRow.h
@@ -72,6 +72,10 @@ public:
 
 	void SetRowForRename();
 
+	void SetIsDragDrop(bool bInIsDragDrop);
+
+	const FDataTableEditorRowListViewDataPtr& GetRowDataPtr() const;
+
 private:
 
 	void OnSearchForReferences();
@@ -81,6 +85,13 @@ private:
 
 	TSharedRef<SWidget> MakeCellWidget(const int32 InRowIndex, const FName& InColumnId);
 
+	void OnRowDragEnter(const FDragDropEvent& DragDropEvent);
+	void OnRowDragLeave(const FDragDropEvent& DragDropEvent);
+
+	virtual const FSlateBrush* GetBorder() const;
+
+	void OnMoveToExtentClicked(FDataTableEditorUtils::ERowMoveDirection MoveDirection);
+
 	TSharedPtr<SInlineEditableTextBlock> InlineEditableText;
 
 	TSharedPtr<FName> CurrentName;
@@ -88,6 +99,8 @@ private:
 	FDataTableEditorRowListViewDataPtr RowDataPtr;
 	TWeakPtr<FDataTableEditor> DataTableEditor;
 
+	bool bIsDragDropObject;
+	bool bIsHoveredDragTarget;
 };
 
 class FDataTableRowDragDropOp : public FDecoratedDragDropOp
@@ -95,7 +108,9 @@ class FDataTableRowDragDropOp : public FDecoratedDragDropOp
 public:
 	DRAG_DROP_OPERATOR_TYPE(FDataTableRowDragDropOp, FDecoratedDragDropOp)
 
-		FDataTableRowDragDropOp(TSharedPtr<SDataTableListViewRow> InRow);
+	FDataTableRowDragDropOp(TSharedPtr<SDataTableListViewRow> InRow);
+
+	void OnDrop(bool bDropWasHandled, const FPointerEvent& MouseEvent);
 
 	TSharedPtr<SWidget> DecoratorWidget;
 
diff --git a/Engine/Source/Editor/DataTableEditor/Public/DataTableRowUtlis.cpp b/Engine/Source/Editor/DataTableEditor/Public/DataTableRowUtlis.cpp
--- a/Engine/Source/Editor/DataTableEditor/Public/DataTableRowUtlis.cpp
+++ b/Engine/Source/Editor/DataTableEditor/Public/DataTableRowUtlis.cpp
@@ -23,11 +23,18 @@ const FText FDataTableRowUtils::InsertNewRowAboveActionTooltip = LOCTEXT("FDataT
 const FText FDataTableRowUtils::InsertNewRowBelowActionName = LOCTEXT("FDataTableRowUtils_InsertNewRowBelow", "Insert New Row Below");
 const FText FDataTableRowUtils::InsertNewRowBelowActionTooltip = LOCTEXT("FDataTableRowUtils_InsertNewRowBelowTooltip", "Insert a new Row below the current selection");
 
+const FText FDataTableRowUtils::MoveToBottomActionName = LOCTEXT("FDataTableRowUtils_MoveToBottom", "Move Row To Bottom");
+const FText FDataTableRowUtils::MoveToBottomActionTooltip = LOCTEXT("FDataTableRowUtils_MoveToBottomTooltip", "Move selected Row to the bottom");
+
+const FText FDataTableRowUtils::MoveToTopActionName = LOCTEXT("FDataTableRowUtils_MoveToTopAction", "Move Row to Top");
+const FText FDataTableRowUtils::MoveToTopActionTooltip = LOCTEXT("FDataTableRowUtils_MoveToTopActionTooltip", "Move selected Row to the top");
+
 
 TSharedRef<SWidget> FDataTableRowUtils::MakeRowActionsMenu(TSharedPtr<IDataTableEditor> Editor, FExecuteAction SearchForReferencesAction, FExecuteAction InsertNewRowAction,
-	FExecuteAction InsertNewRowAboveAction, FExecuteAction InsertNewRowBelowAction)
+	FExecuteAction InsertNewRowAboveAction, FExecuteAction InsertNewRowBelowAction, FExecuteAction MoveToBottomAction, FExecuteAction MoveToTopAction)
 {
-	if (SearchForReferencesAction.IsBound() && InsertNewRowAction.IsBound() && InsertNewRowAboveAction.IsBound() && InsertNewRowBelowAction.IsBound() )
+	if (SearchForReferencesAction.IsBound() && InsertNewRowAction.IsBound() && InsertNewRowAboveAction.IsBound() 
+		&& InsertNewRowBelowAction.IsBound() )
 	{
 		FMenuBuilder MenuBuilder(true, Editor->GetToolkitCommands());
 		MenuBuilder.AddMenuEntry(InsertNewRowActionName, InsertNewRowActionTooltip,
@@ -43,6 +50,13 @@ TSharedRef<SWidget> FDataTableRowUtils::MakeRowActionsMenu(TSharedPtr<IDataTable
 		MenuBuilder.AddMenuEntry(FGenericCommands::Get().Rename);
 		MenuBuilder.AddMenuEntry(FGenericCommands::Get().Delete);
 		MenuBuilder.AddMenuSeparator();
+
+		MenuBuilder.AddMenuEntry(MoveToTopActionName, MoveToTopActionTooltip,
+			FSlateIcon(FEditorStyle::GetStyleSetName(), "Symbols.DoubleUpArrow"), FUIAction(MoveToTopAction));
+		MenuBuilder.AddMenuEntry(MoveToBottomActionName, MoveToBottomActionTooltip,
+			FSlateIcon(FEditorStyle::GetStyleSetName(), "Symbols.DoubleDownArrow"), FUIAction(MoveToBottomAction));
+		MenuBuilder.AddMenuSeparator();
+		
 		MenuBuilder.AddMenuEntry(SearchForReferencesActionName, SearchForReferencesActionTooltip, 
 			FSlateIcon(), FUIAction(SearchForReferencesAction));
 		return MenuBuilder.MakeWidget();
diff --git a/Engine/Source/Editor/DataTableEditor/Public/DataTableRowUtlis.h b/Engine/Source/Editor/DataTableEditor/Public/DataTableRowUtlis.h
--- a/Engine/Source/Editor/DataTableEditor/Public/DataTableRowUtlis.h
+++ b/Engine/Source/Editor/DataTableEditor/Public/DataTableRowUtlis.h
@@ -13,7 +13,8 @@ class DATATABLEEDITOR_API FDataTableRowUtils
 {
 public:
 	static TSharedRef<SWidget> MakeRowActionsMenu(TSharedPtr<class IDataTableEditor> Editor, FExecuteAction SearchForReferencesAction, FExecuteAction InsertNewRowAction,
-													FExecuteAction InsertNewRowAboveAction, FExecuteAction InsertNewRowBelowAction);
+													FExecuteAction InsertNewRowAboveAction, FExecuteAction InsertNewRowBelowAction,
+												    FExecuteAction MoveToBottomAction, FExecuteAction MoveToTopAction);
 	static void AddSearchForReferencesContextMenu(FDetailWidgetRow& RowNameDetailWidget, FExecuteAction SearchForReferencesAction);
 
 private:
@@ -29,6 +30,11 @@ private:
 	static const FText InsertNewRowBelowActionName;
 	static const FText InsertNewRowBelowActionTooltip;
 
+	static const FText MoveToBottomActionName;
+	static const FText MoveToBottomActionTooltip;
+
+	static const FText MoveToTopActionName;
+	static const FText MoveToTopActionTooltip;
 
 };
 
diff --git a/Engine/Source/Editor/EditorStyle/Private/SlateEditorStyle.cpp b/Engine/Source/Editor/EditorStyle/Private/SlateEditorStyle.cpp
--- a/Engine/Source/Editor/EditorStyle/Private/SlateEditorStyle.cpp
+++ b/Engine/Source/Editor/EditorStyle/Private/SlateEditorStyle.cpp
@@ -6847,6 +6847,12 @@ void FSlateEditorStyle::FStyle::SetupToolkitStyles()
 		Set("DataTableEditor.Copy", new IMAGE_BRUSH("/Icons/Edit/icon_Edit_Copy_40x", Icon16x16));
 		Set("DataTableEditor.Paste", new IMAGE_BRUSH("/Icons/Edit/icon_Edit_Paste_40x", Icon16x16));
 		Set("DataTableEditor.Duplicate", new IMAGE_BRUSH("Icons/Edit/icon_Edit_Duplicate_40x", Icon16x16));
+		Set("DataTableEditor.Save", new IMAGE_BRUSH("Icons/icon_SaveAsset_40x", Icon16x16));
+		Set("DataTableEditor.Browse", new IMAGE_BRUSH("Icons/lens_12x", Icon16x16));
+		Set("DataTableEditor.Add", new IMAGE_BRUSH("Icons/PlusSymbol_12x", Icon16x16));
+		Set("DataTableEditor.Remove", new IMAGE_BRUSH("Icons/Cross_12x", Icon16x16));
+
+
 
 		Set( "DataTableEditor.CellText", FTextBlockStyle(NormalText)
 			.SetFont( DEFAULT_FONT("Regular", 9 ))
@@ -6866,18 +6872,22 @@ void FSlateEditorStyle::FStyle::SetupToolkitStyles()
 			.SetSelectedTextColor( DefaultForeground )
 		);
 
-		Set( "DataTableEditor.CellListViewRow", FTableRowStyle(NormalTableRowStyle)
+		Set("DataTableEditor.CellListViewRow", FTableRowStyle(NormalTableRowStyle)
 			.SetEvenRowBackgroundBrush(IMAGE_BRUSH("PropertyView/DetailCategoryMiddle", FVector2D(16, 16), FLinearColor(0.5f, 0.5f, 0.5f)))
 			.SetEvenRowBackgroundHoveredBrush(IMAGE_BRUSH("PropertyView/DetailCategoryMiddle_Hovered", FVector2D(16, 16), FLinearColor(0.5f, 0.5f, 0.5f)))
 			.SetOddRowBackgroundBrush(IMAGE_BRUSH("PropertyView/DetailCategoryMiddle", FVector2D(16, 16), FLinearColor(0.2f, 0.2f, 0.2f)))
 			.SetOddRowBackgroundHoveredBrush(IMAGE_BRUSH("PropertyView/DetailCategoryMiddle_Hovered", FVector2D(16, 16), FLinearColor(0.2f, 0.2f, 0.2f)))
-			.SetActiveBrush( IMAGE_BRUSH( "Common/Selection", Icon8x8, FLinearColor(0.075f, 0.075f, 0.075f) ) )
-			.SetActiveHoveredBrush( IMAGE_BRUSH( "Common/Selection", Icon8x8, FLinearColor(0.075f, 0.075f, 0.075f) ) )
-			.SetInactiveBrush( IMAGE_BRUSH( "Common/Selection", Icon8x8, FLinearColor(0.075f, 0.075f, 0.075f) ) )
-			.SetInactiveHoveredBrush( IMAGE_BRUSH( "Common/Selection", Icon8x8, FLinearColor(0.075f, 0.075f, 0.075f) ) )
-			.SetTextColor( DefaultForeground )
-			.SetSelectedTextColor( DefaultForeground )
+			.SetActiveBrush(IMAGE_BRUSH("Common/Selection", Icon8x8, FLinearColor(0.075f, 0.075f, 0.075f)))
+			.SetActiveHoveredBrush(IMAGE_BRUSH("Common/Selection", Icon8x8, FLinearColor(0.075f, 0.075f, 0.075f)))
+			.SetInactiveBrush(IMAGE_BRUSH("Common/Selection", Icon8x8, FLinearColor(0.075f, 0.075f, 0.075f)))
+			.SetInactiveHoveredBrush(IMAGE_BRUSH("Common/Selection", Icon8x8, FLinearColor(0.075f, 0.075f, 0.075f)))
+			.SetTextColor(DefaultForeground)
+			.SetSelectedTextColor(DefaultForeground)
 		);
+
+		Set("DataTableEditor.DragDropObject", new BOX_BRUSH("Common/TextBox_Special_Active", FMargin(8.0f / 32.0f)));
+		Set("DataTableEditor.DragDropHoveredTarget", new BOX_BRUSH("Common/TextBox_Special_Active", FMargin(8.0f / 32.0f), SelectionColor_Pressed));
+
 	}
 
 	// StringTable Editor
