Added in the standard toolbar, while moving the buttons to it. Fixed a drag drop highlight bug, and added row number on the drag drop decorator.

diff --git a/Engine/Source/Editor/DataTableEditor/Private/DataTableEditor.cpp b/Engine/Source/Editor/DataTableEditor/Private/DataTableEditor.cpp
--- a/Engine/Source/Editor/DataTableEditor/Private/DataTableEditor.cpp
+++ b/Engine/Source/Editor/DataTableEditor/Private/DataTableEditor.cpp
@@ -42,8 +42,6 @@
 #include "Engine/DataTable.h"
 #include "Subsystems/AssetEditorSubsystem.h"
 
-
-
 #define LOCTEXT_NAMESPACE "DataTableEditor"
 
 const FName FDataTableEditor::DataTableTabId("DataTableEditor_DataTable");
@@ -53,10 +51,38 @@ const FName FDataTableEditor::RowNameColumnId("RowName");
 const FName FDataTableEditor::RowNumberColumnId("RowNumber");
 const FName FDataTableEditor::RowDragDropColumnId("RowDragDrop");
 
+class SDataTableModeSeparator : public SBorder
+{
+public:
+	SLATE_BEGIN_ARGS(SDataTableModeSeparator) {}
+	SLATE_END_ARGS()
+
+		void Construct(const FArguments& InArg)
+	{
+		SBorder::Construct(
+			SBorder::FArguments()
+			.BorderImage(FEditorStyle::GetBrush("BlueprintEditor.PipelineSeparator"))
+			.Padding(0.0f)
+		);
+	}
+
+	// SWidget interface
+	virtual FVector2D ComputeDesiredSize(float) const override
+	{
+		const float Height = 20.0f;
+		const float Thickness = 16.0f;
+		return FVector2D(Thickness, Height);
+	}
+	// End of SWidget interface
+};
+
+
 void FDataTableEditor::RegisterTabSpawners(const TSharedRef<class FTabManager>& InTabManager)
 {
 	WorkspaceMenuCategory = InTabManager->AddLocalWorkspaceMenuCategory(LOCTEXT("WorkspaceMenu_Data Table Editor", "Data Table Editor"));
 
+	FAssetEditorToolkit::RegisterTabSpawners(InTabManager);
+
 	CreateAndRegisterDataTableTab(InTabManager);
 	CreateAndRegisterDataTableDetailsTab(InTabManager);
 	CreateAndRegisterRowEditorTab(InTabManager);
@@ -64,6 +90,8 @@ void FDataTableEditor::RegisterTabSpawners(const TSharedRef<class FTabManager>&
 
 void FDataTableEditor::UnregisterTabSpawners(const TSharedRef<class FTabManager>& InTabManager)
 {
+	FAssetEditorToolkit::UnregisterTabSpawners(InTabManager);
+
 	InTabManager->UnregisterTabSpawner(DataTableTabId);
 	InTabManager->UnregisterTabSpawner(DataTableDetailsTabId);
 	InTabManager->UnregisterTabSpawner(RowEditorTabId);
@@ -195,11 +223,18 @@ void FDataTableEditor::HandlePostChange()
 
 void FDataTableEditor::InitDataTableEditor( const EToolkitMode::Type Mode, const TSharedPtr< class IToolkitHost >& InitToolkitHost, UDataTable* Table )
 {
-	TSharedRef<FTabManager::FLayout> StandaloneDefaultLayout = FTabManager::NewLayout( "Standalone_DataTableEditor_Layout_v4" )
+	TSharedRef<FTabManager::FLayout> StandaloneDefaultLayout = FTabManager::NewLayout( "Standalone_DataTableEditor_Layout_v5" )
 	->AddArea
 	(
 		FTabManager::NewPrimaryArea()->SetOrientation(Orient_Vertical)
 		->Split
+		(
+			FTabManager::NewStack()
+			->SetSizeCoefficient(0.1f)
+			->SetHideTabWell(true)
+			->AddTab(GetToolbarTabId(), ETabState::OpenedTab)
+		)
+		->Split
 		(
 			FTabManager::NewStack()
 			->AddTab(DataTableTabId, ETabState::OpenedTab)
@@ -214,11 +249,16 @@ void FDataTableEditor::InitDataTableEditor( const EToolkitMode::Type Mode, const
 	);
 
 	const bool bCreateDefaultStandaloneMenu = true;
-	const bool bCreateDefaultToolbar = false;
+	const bool bCreateDefaultToolbar = true;
 	FAssetEditorToolkit::InitAssetEditor( Mode, InitToolkitHost, FDataTableEditorModule::DataTableEditorAppIdentifier, StandaloneDefaultLayout, bCreateDefaultStandaloneMenu, bCreateDefaultToolbar, Table );
 
 	FDataTableEditorModule& DataTableEditorModule = FModuleManager::LoadModuleChecked<FDataTableEditorModule>( "DataTableEditor" );
 	AddMenuExtender(DataTableEditorModule.GetMenuExtensibilityManager()->GetAllExtenders(GetToolkitCommands(), GetEditingObjects()));
+	
+	TSharedPtr<FExtender> ToolbarExtender = DataTableEditorModule.GetToolBarExtensibilityManager()->GetAllExtenders(GetToolkitCommands(), GetEditingObjects());
+	ExtendToolbar(ToolbarExtender);
+	
+	AddToolbarExtender(ToolbarExtender);
 
 	RegenerateMenusAndToolbars();
 
@@ -261,7 +301,7 @@ FString FDataTableEditor::GetDocumentationLink() const
 	return FString(TEXT("Gameplay/DataDriven"));
 }
 
-FReply FDataTableEditor::OnAddClicked()
+void FDataTableEditor::OnAddClicked()
 {
 	UDataTable* Table = GetEditableDataTable();
 
@@ -276,11 +316,9 @@ FReply FDataTableEditor::OnAddClicked()
 		FDataTableEditorUtils::AddRow(Table, NewName);
 		FDataTableEditorUtils::SelectRow(Table, NewName);
 	}
-
-	return FReply::Handled();
 }
 
-FReply FDataTableEditor::OnRemoveClicked()
+void FDataTableEditor::OnRemoveClicked()
 {
 	UDataTable* Table = GetEditableDataTable();
 
@@ -303,7 +341,6 @@ FReply FDataTableEditor::OnRemoveClicked()
 		}
 
 	}
-	return FReply::Handled();
 }
 
 FReply FDataTableEditor::OnMoveRowClicked(FDataTableEditorUtils::ERowMoveDirection MoveDirection)
@@ -329,34 +366,31 @@ FReply FDataTableEditor::OnMoveToExtentClicked(FDataTableEditorUtils::ERowMoveDi
 	return FReply::Handled();
 }
 
-FReply FDataTableEditor::OnCopyClicked()
+void FDataTableEditor::OnCopyClicked()
 {
 	UDataTable* Table = GetEditableDataTable();
 	if (Table)
 	{
 		CopySelectedRow();
 	}
-	return FReply::Handled();
 }
 
-FReply FDataTableEditor::OnPasteClicked()
+void FDataTableEditor::OnPasteClicked()
 {
 	UDataTable* Table = GetEditableDataTable();
 	if (Table)
 	{
 		PasteOnSelectedRow();
 	}
-	return FReply::Handled();
 }
 
-FReply FDataTableEditor::OnDuplicateClicked()
+void FDataTableEditor::OnDuplicateClicked()
 {
 	UDataTable* Table = GetEditableDataTable();
 	if (Table)
 	{
 		DuplicateSelectedRow();
 	}
-	return FReply::Handled();
 }
 
 EColumnSortMode::Type FDataTableEditor::GetColumnSortMode(const FName ColumnId) const
@@ -497,6 +531,56 @@ FReply FDataTableEditor::BrowseForDataTable_Execute()
 	return FReply::Handled();
 }
 
+void FDataTableEditor::ExtendToolbar(TSharedPtr<FExtender> Extender)
+{
+	Extender->AddToolBarExtension(
+		"Asset",
+		EExtensionHook::After,
+		GetToolkitCommands(),
+		FToolBarExtensionDelegate::CreateSP(this, &FDataTableEditor::FillToolbar)
+	);
+
+}
+
+void FDataTableEditor::FillToolbar(FToolBarBuilder& ToolbarBuilder)
+{
+	ToolbarBuilder.BeginSection("DataTableCommands");
+	{
+		ToolbarBuilder.AddToolBarButton(
+			FUIAction(FExecuteAction::CreateSP(this, &FDataTableEditor::OnAddClicked)),
+			NAME_None,
+			LOCTEXT("AddIconText", "Add"),
+			LOCTEXT("AddRowToolTip", "Add a new row to the Data Table"),
+			FSlateIcon(FEditorStyle::GetStyleSetName(), "DataTableEditor.Add"));
+		ToolbarBuilder.AddToolBarButton(
+			FUIAction(FExecuteAction::CreateSP(this, &FDataTableEditor::OnCopyClicked)),
+			NAME_None,
+			LOCTEXT("CopyIconText", "Copy"),
+			LOCTEXT("CopyToolTip", "Copy the currently selected row"),
+			FSlateIcon(FEditorStyle::GetStyleSetName(), "DataTableEditor.Copy"));
+		ToolbarBuilder.AddToolBarButton(
+			FUIAction(FExecuteAction::CreateSP(this, &FDataTableEditor::OnPasteClicked)),
+			NAME_None,
+			LOCTEXT("PasteIconText", "Paste"),
+			LOCTEXT("PasteToolTip", "Paste on the currently selected row"),
+			FSlateIcon(FEditorStyle::GetStyleSetName(), "DataTableEditor.Paste"));
+		ToolbarBuilder.AddToolBarButton(
+			FUIAction(FExecuteAction::CreateSP(this, &FDataTableEditor::OnDuplicateClicked)),
+			NAME_None,
+			LOCTEXT("DuplicateIconText", "Duplicate"),
+			LOCTEXT("DuplicateToolTip", "Duplicate the currently selected row"),
+			FSlateIcon(FEditorStyle::GetStyleSetName(), "DataTableEditor.Duplicate"));
+		ToolbarBuilder.AddToolBarButton(
+			FUIAction(FExecuteAction::CreateSP(this, &FDataTableEditor::OnRemoveClicked)),
+			NAME_None,
+			LOCTEXT("RemoveRowIconText", "Remove"),
+			LOCTEXT("RemoveRowToolTip", "Remove the currently selected row from the Data Table"),
+			FSlateIcon(FEditorStyle::GetStyleSetName(), "DataTableEditor.Remove"));
+	}
+	ToolbarBuilder.EndSection();
+
+}
+
 UDataTable* FDataTableEditor::GetEditableDataTable() const
 {
 	return Cast<UDataTable>(GetEditingObject());
@@ -1178,125 +1262,6 @@ TSharedRef<SVerticalBox> FDataTableEditor::CreateContentBox()
 		[
 			SNew(SHorizontalBox)
 			+ SHorizontalBox::Slot()
-			.AutoWidth()
-			.Padding(2)
-			[
-				SNew(SButton)
-				.ButtonStyle(FEditorStyle::Get(), "HoverHintOnly")
-				.ForegroundColor(FSlateColor::UseForeground())
-				.HAlign(HAlign_Center)
-				.VAlign(VAlign_Center)
-				.OnClicked(this, &FDataTableEditor::SaveDataTable_Execute)
-				.ToolTipText(LOCTEXT("SaveDataTableTooltip", "Save data table"))
-				[
-					SNew(SImage)
-					.Image(FEditorStyle::Get().GetBrush("AssetEditor.SaveAsset.Small"))
-				]
-			]
-			+ SHorizontalBox::Slot()
-			.AutoWidth()
-			.Padding(2)
-			[
-				SNew(SButton)
-				.ButtonStyle(FEditorStyle::Get(), "HoverHintOnly")
-				.ForegroundColor(FSlateColor::UseForeground())
-				.HAlign(HAlign_Center)
-				.VAlign(VAlign_Center)
-				.OnClicked(this, &FDataTableEditor::BrowseForDataTable_Execute)
-				.ToolTipText(LOCTEXT("FindInContentBrowserTooltip", "Find data table in content browser"))
-				[
-					SNew(SImage)
-					.Image(FEditorStyle::Get().GetBrush("Sequencer.FindInContentBrowser.Small"))
-				]
-			]
-			+ SHorizontalBox::Slot()
-			.AutoWidth()
-			.Padding(2)
-			[
-				SNew(SButton)
-				.ButtonStyle(FEditorStyle::Get(), "HoverHintOnly")
-				.ForegroundColor(FSlateColor::UseForeground())
-				.HAlign(HAlign_Center)
-				.VAlign(VAlign_Center)
-				.OnClicked(this, &FDataTableEditor::OnAddClicked)
-				.ToolTipText(LOCTEXT("AddRowTooltip", "Add a new row to the data table"))
-				[
-					SNew(SImage)
-					.Image(FEditorStyle::Get().GetBrush("DataTableEditor.Add"))
-				]
-			]
-			+ SHorizontalBox::Slot()
-			.AutoWidth()
-			.Padding(2)
-			[
-				SNew(SButton)
-				.ButtonStyle(FEditorStyle::Get(), "HoverHintOnly")
-				.ForegroundColor(FSlateColor::UseForeground())
-				.HAlign(HAlign_Center)
-				.VAlign(VAlign_Center)
-				.OnClicked(this, &FDataTableEditor::OnCopyClicked)
-				.ToolTipText(LOCTEXT("CopyTooltip", "Copy the currently selected row"))
-				[
-					SNew(SImage)
-					.Image(FEditorStyle::Get().GetBrush("DataTableEditor.Copy"))
-				]
-			]
-			+ SHorizontalBox::Slot()
-			.AutoWidth()
-			.Padding(2)
-			[
-				SNew(SButton)
-				.ButtonStyle(FEditorStyle::Get(), "HoverHintOnly")
-				.ForegroundColor(FSlateColor::UseForeground())
-				.HAlign(HAlign_Center)
-				.VAlign(VAlign_Center)
-				.OnClicked(this, &FDataTableEditor::OnPasteClicked)
-				.ToolTipText(LOCTEXT("PasteTooltip", "Paste on the currently selected row"))
-				[
-					SNew(SImage)
-					.Image(FEditorStyle::Get().GetBrush("DataTableEditor.Paste"))
-				]
-			]
-			+ SHorizontalBox::Slot()
-			.AutoWidth()
-			.Padding(2)
-			[
-				SNew(SButton)
-				.ButtonStyle(FEditorStyle::Get(), "HoverHintOnly")
-				.ForegroundColor(FSlateColor::UseForeground())
-				.HAlign(HAlign_Center)
-				.VAlign(VAlign_Center)
-				.OnClicked(this, &FDataTableEditor::OnDuplicateClicked)
-				.ToolTipText(LOCTEXT("DuplicateTooltip", "Duplicate the currently selected row"))
-				[
-					SNew(SImage)
-					.Image(FEditorStyle::Get().GetBrush("DataTableEditor.Duplicate"))
-				]
-			]
-			+ SHorizontalBox::Slot()
-			.AutoWidth()
-			.Padding(2)
-			[
-				SNew(SButton)
-				.ButtonStyle(FEditorStyle::Get(), "HoverHintOnly")
-				.ForegroundColor(FSlateColor::UseForeground())
-				.HAlign(HAlign_Center)
-				.VAlign(VAlign_Center)
-				.OnClicked(this, &FDataTableEditor::OnRemoveClicked)
-				.ToolTipText(LOCTEXT("RemoveRowTooltip", "Remove the currently selected row from the data table"))
-				[
-					SNew(SImage)
-					.Image(FEditorStyle::Get().GetBrush("DataTableEditor.Remove"))
-				]
-			]
-			+ SHorizontalBox::Slot()
-			.AutoWidth()
-			.Padding(2)
-			[
-				SNew(SSeparator)
-				.Orientation(Orient_Vertical)
-			]
-			+ SHorizontalBox::Slot()
 			[
 				SAssignNew(SearchBoxWidget, SSearchBox)
 				.InitialText(this, &FDataTableEditor::GetFilterText)
diff --git a/Engine/Source/Editor/DataTableEditor/Private/DataTableEditor.h b/Engine/Source/Editor/DataTableEditor/Private/DataTableEditor.h
--- a/Engine/Source/Editor/DataTableEditor/Private/DataTableEditor.h
+++ b/Engine/Source/Editor/DataTableEditor/Private/DataTableEditor.h
@@ -164,13 +164,13 @@ protected:
 
 	virtual FString GetDocumentationLink() const override;
 	
-	FReply OnAddClicked();
-	FReply OnRemoveClicked();
+	void OnAddClicked();
+	void OnRemoveClicked();
 	FReply OnMoveRowClicked(FDataTableEditorUtils::ERowMoveDirection MoveDirection);
 	FReply OnMoveToExtentClicked(FDataTableEditorUtils::ERowMoveDirection MoveDirection);
-	FReply OnCopyClicked();
-	FReply OnPasteClicked();
-	FReply OnDuplicateClicked();
+	void OnCopyClicked();
+	void OnPasteClicked();
+	void OnDuplicateClicked();
 
 	EColumnSortMode::Type GetColumnSortMode(const FName ColumnId) const;
 	void OnColumnSortModeChanged(const EColumnSortPriority::Type SortPriority, const FName& ColumnId, const EColumnSortMode::Type InSortMode);
@@ -181,7 +181,9 @@ protected:
 
 	FReply SaveDataTable_Execute();
 	FReply BrowseForDataTable_Execute();
-
+	
+	void ExtendToolbar(TSharedPtr<FExtender> Extender);
+	void FillToolbar(FToolBarBuilder& ToolbarBuilder);
 
 private:
 	UDataTable* GetEditableDataTable() const;
diff --git a/Engine/Source/Editor/DataTableEditor/Private/DataTableEditorModule.cpp b/Engine/Source/Editor/DataTableEditor/Private/DataTableEditorModule.cpp
--- a/Engine/Source/Editor/DataTableEditor/Private/DataTableEditorModule.cpp
+++ b/Engine/Source/Editor/DataTableEditor/Private/DataTableEditorModule.cpp
@@ -15,12 +15,14 @@ const FName FDataTableEditorModule::DataTableEditorAppIdentifier( TEXT( "DataTab
 void FDataTableEditorModule::StartupModule()
 {
 	MenuExtensibilityManager = MakeShareable(new FExtensibilityManager);
+	ToolBarExtensibilityManager = MakeShareable(new FExtensibilityManager);
 }
 
 
 void FDataTableEditorModule::ShutdownModule()
 {
 	MenuExtensibilityManager.Reset();
+	ToolBarExtensibilityManager.Reset();
 }
 
 
diff --git a/Engine/Source/Editor/DataTableEditor/Private/SDataTableListViewRow.cpp b/Engine/Source/Editor/DataTableEditor/Private/SDataTableListViewRow.cpp
--- a/Engine/Source/Editor/DataTableEditor/Private/SDataTableListViewRow.cpp
+++ b/Engine/Source/Editor/DataTableEditor/Private/SDataTableListViewRow.cpp
@@ -100,6 +100,8 @@ void SDataTableListViewRow::OnInsertNewRow(ERowInsertionPosition InsertPosition)
 
 FReply SDataTableListViewRow::OnRowDrop(const FDragDropEvent& DragDropEvent)
 {
+	bIsHoveredDragTarget = false;
+
 	TSharedPtr<FDataTableRowDragDropOp> DataTableDropOp = DragDropEvent.GetOperationAs< FDataTableRowDragDropOp >();
 	TSharedPtr<SDataTableListViewRow> RowPtr = nullptr;
 	if (DataTableDropOp.IsValid() && DataTableDropOp->Row.IsValid())
@@ -300,6 +302,17 @@ FName SDataTableListViewRow::GetCurrentName() const
 
 }
 
+uint32 SDataTableListViewRow::GetCurrentIndex() const
+{
+	return RowDataPtr.IsValid() ? RowDataPtr->RowNum : -1;
+}
+
+const TArray<FText>& SDataTableListViewRow::GetCellValues() const
+{
+	check(RowDataPtr)
+	return RowDataPtr->CellData;
+}
+
 FReply SDataTableListViewRow::OnMouseButtonDoubleClick(const FGeometry& InMyGeometry, const FPointerEvent& InMouseEvent)
 {
 	if (InlineEditableText->IsHovered())
@@ -408,24 +421,24 @@ FDataTableRowDragDropOp::FDataTableRowDragDropOp(TSharedPtr<SDataTableListViewRo
 	{
 		RowPtr = Row.Pin();
 		RowPtr->SetIsDragDrop(true);
-	}
 
-	DecoratorWidget = SNew(SBorder)
-		.Padding(8.f)
-		.BorderImage(FEditorStyle::GetBrush("Graph.ConnectorFeedback.Border"))
-		.Content()
-		[
-			SNew(SHorizontalBox)
-			+ SHorizontalBox::Slot()
-			.AutoWidth()
-			.VAlign(VAlign_Center)
+		DecoratorWidget = SNew(SBorder)
+			.Padding(8.f)
+			.BorderImage(FEditorStyle::GetBrush("Graph.ConnectorFeedback.Border"))
+			.Content()
 			[
-				SNew(STextBlock)
-				.Text(FText::Format(LOCTEXT("DragDropDecoratorText", "Place {0} Here"), FText::FromName(InRow->GetCurrentName())))
-			]
-		];
+				SNew(SHorizontalBox)
+				+ SHorizontalBox::Slot()
+				.AutoWidth()
+				.VAlign(VAlign_Center)
+				[
+					SNew(STextBlock)
+					.Text(FText::Format(NSLOCTEXT("DataTableDragDrop", "PlaceRowHere", "Place Row {0} Here"), FText::AsNumber(InRow->GetCurrentIndex())))
+				]
+			];
 
-	Construct();
+		Construct();
+	}
 }
 
 void FDataTableRowDragDropOp::OnDrop(bool bDropWasHandled, const FPointerEvent& MouseEvent)
diff --git a/Engine/Source/Editor/DataTableEditor/Private/SDataTableListViewRow.h b/Engine/Source/Editor/DataTableEditor/Private/SDataTableListViewRow.h
--- a/Engine/Source/Editor/DataTableEditor/Private/SDataTableListViewRow.h
+++ b/Engine/Source/Editor/DataTableEditor/Private/SDataTableListViewRow.h
@@ -67,6 +67,9 @@ public:
 
 	FText GetCurrentNameAsText() const;
 	FName GetCurrentName() const;
+	uint32 GetCurrentIndex() const;
+
+	const TArray<FText>& GetCellValues() const;
 
 	virtual FReply OnMouseButtonDoubleClick(const FGeometry& InMyGeometry, const FPointerEvent& InMouseEvent);
 
@@ -120,4 +123,5 @@ public:
 	}
 
 	TWeakPtr<class SDataTableListViewRow> Row;
+
 };
diff --git a/Engine/Source/Editor/DataTableEditor/Public/DataTableEditorModule.h b/Engine/Source/Editor/DataTableEditor/Public/DataTableEditorModule.h
--- a/Engine/Source/Editor/DataTableEditor/Public/DataTableEditorModule.h
+++ b/Engine/Source/Editor/DataTableEditor/Public/DataTableEditorModule.h
@@ -12,7 +12,7 @@ class UDataTable;
 
 /** DataTable Editor module */
 class FDataTableEditorModule : public IModuleInterface,
-	public IHasMenuExtensibility
+	public IHasMenuExtensibility, public IHasToolBarExtensibility
 {
 
 public:
@@ -34,6 +34,8 @@ public:
 	/** Gets the extensibility managers for outside entities to extend static mesh editor's menus and toolbars */
 	virtual TSharedPtr<FExtensibilityManager> GetMenuExtensibilityManager() override {return MenuExtensibilityManager;}
 
+	virtual TSharedPtr<FExtensibilityManager> GetToolBarExtensibilityManager() override { return ToolBarExtensibilityManager; }
+	
 	/** DataTable Editor app identifier string */
 	static const FName DataTableEditorAppIdentifier;
 
@@ -43,6 +45,9 @@ private:
 	TSharedRef<IDataTableEditor> CreateCompositeDataTableEditor(const EToolkitMode::Type Mode, const TSharedPtr< class IToolkitHost >& InitToolkitHost, UDataTable* Table);
 
 	TSharedPtr<FExtensibilityManager> MenuExtensibilityManager;
+	TSharedPtr<FExtensibilityManager> ToolBarExtensibilityManager;
+;
+
 };
 
 
diff --git a/Engine/Source/Editor/DataTableEditor/Public/DataTableRowUtlis.cpp b/Engine/Source/Editor/DataTableEditor/Public/DataTableRowUtlis.cpp
--- a/Engine/Source/Editor/DataTableEditor/Public/DataTableRowUtlis.cpp
+++ b/Engine/Source/Editor/DataTableEditor/Public/DataTableRowUtlis.cpp
@@ -38,7 +38,7 @@ TSharedRef<SWidget> FDataTableRowUtils::MakeRowActionsMenu(TSharedPtr<IDataTable
 	{
 		FMenuBuilder MenuBuilder(true, Editor->GetToolkitCommands());
 		MenuBuilder.AddMenuEntry(InsertNewRowActionName, InsertNewRowActionTooltip,
-			FSlateIcon(FEditorStyle::GetStyleSetName(), "Plus"), FUIAction(InsertNewRowAction));
+			FSlateIcon(FEditorStyle::GetStyleSetName(), "DataTableEditor.Add"), FUIAction(InsertNewRowAction));
 		MenuBuilder.AddMenuEntry(InsertNewRowAboveActionName, InsertNewRowAboveActionTooltip,
 			FSlateIcon(), FUIAction(InsertNewRowAboveAction));
 		MenuBuilder.AddMenuEntry(InsertNewRowBelowActionName, InsertNewRowBelowActionTooltip,
diff --git a/Engine/Source/Editor/EditorStyle/Private/SlateEditorStyle.cpp b/Engine/Source/Editor/EditorStyle/Private/SlateEditorStyle.cpp
--- a/Engine/Source/Editor/EditorStyle/Private/SlateEditorStyle.cpp
+++ b/Engine/Source/Editor/EditorStyle/Private/SlateEditorStyle.cpp
@@ -6854,13 +6854,19 @@ void FSlateEditorStyle::FStyle::SetupToolkitStyles()
 	{
 		Set( "DataTableEditor.Tabs.Properties", new IMAGE_BRUSH( "/Icons/icon_tab_SelectionDetails_16x", Icon16x16 ) );
 
-		Set("DataTableEditor.Copy", new IMAGE_BRUSH("/Icons/Edit/icon_Edit_Copy_40x", Icon16x16));
-		Set("DataTableEditor.Paste", new IMAGE_BRUSH("/Icons/Edit/icon_Edit_Paste_40x", Icon16x16));
-		Set("DataTableEditor.Duplicate", new IMAGE_BRUSH("Icons/Edit/icon_Edit_Duplicate_40x", Icon16x16));
+		Set("DataTableEditor.Copy", new IMAGE_BRUSH("/Icons/Edit/icon_Edit_Copy_40x", Icon32x32));
+		Set("DataTableEditor.Paste", new IMAGE_BRUSH("/Icons/Edit/icon_Edit_Paste_40x", Icon32x32));
+		Set("DataTableEditor.Duplicate", new IMAGE_BRUSH("Icons/Edit/icon_Edit_Duplicate_40x", Icon32x32));
 		Set("DataTableEditor.Save", new IMAGE_BRUSH("Icons/icon_SaveAsset_40x", Icon16x16));
 		Set("DataTableEditor.Browse", new IMAGE_BRUSH("Icons/lens_12x", Icon16x16));
-		Set("DataTableEditor.Add", new IMAGE_BRUSH("Icons/PlusSymbol_12x", Icon16x16));
-		Set("DataTableEditor.Remove", new IMAGE_BRUSH("Icons/Cross_12x", Icon16x16));
+		Set("DataTableEditor.Add", new IMAGE_BRUSH("Icons/icon_add_40x", Icon32x32));
+		Set("DataTableEditor.Remove", new IMAGE_BRUSH("Icons/Edit/icon_Edit_Delete_40x", Icon32x32));
+
+		Set("DataTableEditor.Copy.Small", new IMAGE_BRUSH("/Icons/Edit/icon_Edit_Copy_40x", Icon16x16));
+		Set("DataTableEditor.Paste.Small", new IMAGE_BRUSH("/Icons/Edit/icon_Edit_Paste_40x", Icon16x16));
+		Set("DataTableEditor.Duplicate.Small", new IMAGE_BRUSH("Icons/Edit/icon_Edit_Duplicate_40x", Icon16x16));
+		Set("DataTableEditor.Add.Small", new IMAGE_BRUSH("Icons/icon_add_40x", Icon16x16));
+		Set("DataTableEditor.Remove.Small", new IMAGE_BRUSH("Icons/Edit/icon_Edit_Delete_40x", Icon16x16));
 
 
 
