Fixed context menu row arrangement not updating before click and took out unnecessary const_casts in the Data Tables.

diff --git a/Engine/Source/Editor/DataTableEditor/Private/DataTableEditor.cpp b/Engine/Source/Editor/DataTableEditor/Private/DataTableEditor.cpp
--- a/Engine/Source/Editor/DataTableEditor/Private/DataTableEditor.cpp
+++ b/Engine/Source/Editor/DataTableEditor/Private/DataTableEditor.cpp
@@ -200,11 +200,11 @@ void FDataTableEditor::PreChange(const UDataTable* Changed, FDataTableEditorUtil
 
 void FDataTableEditor::PostChange(const UDataTable* Changed, FDataTableEditorUtils::EDataTableChangeInfo Info)
 {
-	const UDataTable* Table = GetDataTable();
+	UDataTable* Table = GetEditableDataTable();
 	if (Changed == Table)
 	{
 		HandlePostChange();
-		const_cast<UDataTable*>(Table)->OnDataTableChanged().Broadcast();
+		Table->OnDataTableChanged().Broadcast();
 	}
 }
 
@@ -362,6 +362,7 @@ FReply FDataTableEditor::OnMoveToExtentClicked(FDataTableEditorUtils::ERowMoveDi
 	{
 		// We move by the row map size, as FDataTableEditorUtils::MoveRow will automatically clamp this as appropriate
 		FDataTableEditorUtils::MoveRow(Table, HighlightedRowName, MoveDirection, Table->GetRowMap().Num());
+		FDataTableEditorUtils::SelectRow(Table, HighlightedRowName);
 	}
 	return FReply::Handled();
 }
@@ -1031,7 +1032,7 @@ void FDataTableEditor::OnNavigateToDataTableRowCode()
 
 void FDataTableEditor::RefreshCachedDataTable(const FName InCachedSelection, const bool bUpdateEvenIfValid)
 {
-	const UDataTable* Table = GetDataTable();
+	UDataTable* Table = GetEditableDataTable();
 	TArray<FDataTableEditorColumnHeaderDataPtr> PreviousColumns = AvailableColumns;
 
 	FDataTableEditorUtils::CacheDataTableForEditing(Table, AvailableColumns, AvailableRows);
@@ -1154,7 +1155,7 @@ void FDataTableEditor::RefreshCachedDataTable(const FName InCachedSelection, con
 
 	if (PropertyView.IsValid())
 	{
-		PropertyView->SetObject(const_cast<UDataTable*>(Table));
+		PropertyView->SetObject(Table);
 	}
 }
 
@@ -1374,7 +1375,7 @@ TSharedRef<SDockTab> FDataTableEditor::SpawnTab_DataTableDetails(const FSpawnTab
 {
 	check(Args.GetTabId().TabType == DataTableDetailsTabId);
 
-	PropertyView->SetObject(const_cast<UDataTable*>(GetDataTable()));
+	PropertyView->SetObject(GetEditableDataTable());
 
 	return SNew(SDockTab)
 		.Icon(FEditorStyle::GetBrush("DataTableEditor.Tabs.Properties"))
