Ensured the row name and row number columns did not resize after manual adjustment, refactored the context menu based on current use, and added the row editor details panel to have the same name as the currently selected row, for Data Tables.

diff --git a/Engine/Source/Editor/DataTableEditor/Private/DataTableEditor.cpp b/Engine/Source/Editor/DataTableEditor/Private/DataTableEditor.cpp
--- a/Engine/Source/Editor/DataTableEditor/Private/DataTableEditor.cpp
+++ b/Engine/Source/Editor/DataTableEditor/Private/DataTableEditor.cpp
@@ -649,7 +649,7 @@ void FDataTableEditor::RefreshRowNumberColumnWidth()
 	TSharedRef<FSlateFontMeasure> FontMeasure = FSlateApplication::Get().GetRenderer()->GetFontMeasureService();
 	const FTextBlockStyle& CellTextStyle = FEditorStyle::GetWidgetStyle<FTextBlockStyle>("DataTableEditor.CellText");
 	const float CellPadding = 10.0f;
-	RowNumberColumnWidth = 10.0f;
+
 	for (const FDataTableEditorRowListViewDataPtr& RowData : AvailableRows)
 	{
 		const float RowNumberWidth = FontMeasure->Measure(FString::FromInt(RowData->RowNum), CellTextStyle.Font).X + CellPadding;
@@ -675,7 +675,6 @@ void FDataTableEditor::RefreshRowNameColumnWidth()
 	const FTextBlockStyle& CellTextStyle = FEditorStyle::GetWidgetStyle<FTextBlockStyle>("DataTableEditor.CellText");
 	static const float CellPadding = 10.0f;
 
-	RowNameColumnWidth = 10.0f;
 	for (const FDataTableEditorRowListViewDataPtr& RowData : AvailableRows)
 	{
 		const float RowNameWidth = FontMeasure->Measure(RowData->DisplayName, CellTextStyle.Font).X + CellPadding;
diff --git a/Engine/Source/Editor/DataTableEditor/Private/SDataTableListViewRow.cpp b/Engine/Source/Editor/DataTableEditor/Private/SDataTableListViewRow.cpp
--- a/Engine/Source/Editor/DataTableEditor/Private/SDataTableListViewRow.cpp
+++ b/Engine/Source/Editor/DataTableEditor/Private/SDataTableListViewRow.cpp
@@ -9,6 +9,8 @@
 #include "Framework/MultiBox/MultiBoxBuilder.h"
 #include "Widgets/Text/SInlineEditableTextBlock.h"
 #include "Misc/MessageDialog.h"
+#include "Framework/Commands/GenericCommands.h"
+#include "DetailWidgetRow.h"
 
 #define LOCTEXT_NAMESPACE "SDataTableListViewRowName"
 
@@ -35,14 +37,8 @@ FReply SDataTableListViewRow::OnMouseButtonUp(const FGeometry& MyGeometry, const
 	if (MouseEvent.GetEffectingButton() == EKeys::RightMouseButton && RowDataPtr.IsValid() && FEditorDelegates::OnOpenReferenceViewer.IsBound() && DataTableEditor.IsValid())
 	{
 		FDataTableEditorUtils::SelectRow(DataTableEditor.Pin()->GetDataTable(), RowDataPtr->RowId);
-		TSharedRef<SWidget> MenuWidget = FDataTableRowUtils::MakeRowActionsMenu(DataTableEditor.Pin(),
-			FExecuteAction::CreateSP(this, &SDataTableListViewRow::OnSearchForReferences),
-			FExecuteAction::CreateSP(this, &SDataTableListViewRow::OnInsertNewRow, ERowInsertionPosition::Bottom),
-			FExecuteAction::CreateSP(this, &SDataTableListViewRow::OnInsertNewRow, ERowInsertionPosition::Above),
-			FExecuteAction::CreateSP(this, &SDataTableListViewRow::OnInsertNewRow, ERowInsertionPosition::Below),
-			FExecuteAction::CreateSP(this, &SDataTableListViewRow::OnMoveToExtentClicked, FDataTableEditorUtils::ERowMoveDirection::Down),
-			FExecuteAction::CreateSP(this, &SDataTableListViewRow::OnMoveToExtentClicked, FDataTableEditorUtils::ERowMoveDirection::Up)
-		);
+
+		TSharedRef<SWidget> MenuWidget = MakeRowActionsMenu();
 
 		FWidgetPath WidgetPath = MouseEvent.GetEventPath() != nullptr ? *MouseEvent.GetEventPath() : FWidgetPath();
 		FSlateApplication::Get().PushMenu(AsShared(), WidgetPath, MenuWidget, MouseEvent.GetScreenSpacePosition(), FPopupTransitionEffect::ContextMenu);
@@ -412,6 +408,66 @@ void SDataTableListViewRow::OnMoveToExtentClicked(FDataTableEditorUtils::ERowMov
 	}
 }
 
+
+TSharedRef<SWidget> SDataTableListViewRow::MakeRowActionsMenu()
+{
+	FMenuBuilder MenuBuilder(true, DataTableEditor.Pin()->GetToolkitCommands());
+
+	MenuBuilder.AddMenuEntry(
+		LOCTEXT("DataTableRowMenuActions_InsertNewRow", "Insert New Row"),
+		LOCTEXT("DataTableRowMenuActions_InsertNewRowTooltip", "Insert a new row"),
+		FSlateIcon(FEditorStyle::GetStyleSetName(), "DataTableEditor.Add"), 
+		FUIAction(FExecuteAction::CreateSP(this, &SDataTableListViewRow::OnInsertNewRow, ERowInsertionPosition::Bottom))
+	);
+	
+	MenuBuilder.AddMenuEntry(
+		LOCTEXT("DataTableRowMenuActions_InsertNewRowAbove", "Insert New Row Above"),
+		LOCTEXT("DataTableRowMenuActions_InsertNewRowAboveTooltip", "Insert a new Row above the current selection"),
+		FSlateIcon(), 
+		FUIAction(FExecuteAction::CreateSP(this, &SDataTableListViewRow::OnInsertNewRow, ERowInsertionPosition::Above))
+	);
+	
+	MenuBuilder.AddMenuEntry(
+		LOCTEXT("DataTableRowMenuActions_InsertNewRowBelow", "Insert New Row Below"),
+		LOCTEXT("DataTableRowMenuActions_InsertNewRowBelowTooltip", "Insert a new Row below the current selection"),
+		FSlateIcon(), 
+		FUIAction(FExecuteAction::CreateSP(this, &SDataTableListViewRow::OnInsertNewRow, ERowInsertionPosition::Below))
+	);
+
+	MenuBuilder.AddMenuEntry(FGenericCommands::Get().Copy);
+	MenuBuilder.AddMenuEntry(FGenericCommands::Get().Paste);
+	MenuBuilder.AddMenuEntry(FGenericCommands::Get().Duplicate);
+	MenuBuilder.AddMenuEntry(FGenericCommands::Get().Rename);
+	MenuBuilder.AddMenuEntry(FGenericCommands::Get().Delete);
+
+	MenuBuilder.AddMenuSeparator();
+
+	MenuBuilder.AddMenuEntry(
+		LOCTEXT("DataTableRowMenuActions_MoveToTopAction", "Move Row to Top"),
+		LOCTEXT("DataTableRowMenuActions_MoveToTopActionTooltip", "Move selected Row to the top"),
+		FSlateIcon(FEditorStyle::GetStyleSetName(), "Symbols.DoubleUpArrow"), 
+		FUIAction(FExecuteAction::CreateSP(this, &SDataTableListViewRow::OnMoveToExtentClicked, FDataTableEditorUtils::ERowMoveDirection::Up))
+	);
+	
+	MenuBuilder.AddMenuEntry(
+		LOCTEXT("DataTableRowMenuActions_MoveToBottom", "Move Row To Bottom"),
+		LOCTEXT("DataTableRowMenuActions_MoveToBottomTooltip", "Move selected Row to the bottom"),
+		FSlateIcon(FEditorStyle::GetStyleSetName(), "Symbols.DoubleDownArrow"), 
+		FUIAction(FExecuteAction::CreateSP(this, &SDataTableListViewRow::OnMoveToExtentClicked, FDataTableEditorUtils::ERowMoveDirection::Down))
+	);
+	
+	MenuBuilder.AddMenuSeparator();
+
+	MenuBuilder.AddMenuEntry(
+		NSLOCTEXT("FDataTableRowUtils", "FDataTableRowUtils_SearchForReferences", "Find Row References"),
+		NSLOCTEXT("FDataTableRowUtils", "FDataTableRowUtils_SearchForReferencesTooltip", "Find assets that reference this Row"),
+		FSlateIcon(), 
+		FUIAction(FExecuteAction::CreateSP(this, &SDataTableListViewRow::OnSearchForReferences))
+	);
+
+	return MenuBuilder.MakeWidget();
+}
+
 FDataTableRowDragDropOp::FDataTableRowDragDropOp(TSharedPtr<SDataTableListViewRow> InRow)
 {
 	Row = InRow;
diff --git a/Engine/Source/Editor/DataTableEditor/Private/SDataTableListViewRow.h b/Engine/Source/Editor/DataTableEditor/Private/SDataTableListViewRow.h
--- a/Engine/Source/Editor/DataTableEditor/Private/SDataTableListViewRow.h
+++ b/Engine/Source/Editor/DataTableEditor/Private/SDataTableListViewRow.h
@@ -95,6 +95,8 @@ private:
 
 	void OnMoveToExtentClicked(FDataTableEditorUtils::ERowMoveDirection MoveDirection);
 
+	TSharedRef<SWidget> MakeRowActionsMenu();
+
 	TSharedPtr<SInlineEditableTextBlock> InlineEditableText;
 
 	TSharedPtr<FName> CurrentName;
diff --git a/Engine/Source/Editor/DataTableEditor/Private/SRowEditor.cpp b/Engine/Source/Editor/DataTableEditor/Private/SRowEditor.cpp
--- a/Engine/Source/Editor/DataTableEditor/Private/SRowEditor.cpp
+++ b/Engine/Source/Editor/DataTableEditor/Private/SRowEditor.cpp
@@ -198,6 +198,7 @@ void SRowEditor::Restore()
 	CurrentRow = MakeShareable(new FStructFromDataTable(DataTable.Get(), FinalName));
 	if (StructureDetailsView.IsValid())
 	{
+		StructureDetailsView->SetCustomName(FText::FromName(FinalName));
 		StructureDetailsView->SetStructureData(CurrentRow);
 	}
 
diff --git a/Engine/Source/Editor/DataTableEditor/Public/DataTableRowUtlis.cpp b/Engine/Source/Editor/DataTableEditor/Public/DataTableRowUtlis.cpp
--- a/Engine/Source/Editor/DataTableEditor/Public/DataTableRowUtlis.cpp
+++ b/Engine/Source/Editor/DataTableEditor/Public/DataTableRowUtlis.cpp
@@ -11,66 +11,14 @@
 
 #define LOCTEXT_NAMESPACE "FDataTableRowUtils"
 
-const FText FDataTableRowUtils::SearchForReferencesActionName = LOCTEXT("FDataTableRowUtils_SearchForReferences", "Find Row References");
-const FText FDataTableRowUtils::SearchForReferencesActionTooltip = LOCTEXT("FDataTableRowUtils_SearchForReferencesTooltip", "Find assets that reference this Row");
-
-const FText FDataTableRowUtils::InsertNewRowActionName = LOCTEXT("FDataTableRowUtils_InsertNewRow", "Insert New Row");
-const FText FDataTableRowUtils::InsertNewRowActionTooltip = LOCTEXT("FDataTableRowUtils_InsertNewRowTooltip", "Insert a new Row");
-
-const FText FDataTableRowUtils::InsertNewRowAboveActionName = LOCTEXT("FDataTableRowUtils_InsertNewRowAbove", "Insert New Row Above");
-const FText FDataTableRowUtils::InsertNewRowAboveActionTooltip = LOCTEXT("FDataTableRowUtils_InsertNewRowAboveTooltip", "Insert a new Row above the current selection");
-
-const FText FDataTableRowUtils::InsertNewRowBelowActionName = LOCTEXT("FDataTableRowUtils_InsertNewRowBelow", "Insert New Row Below");
-const FText FDataTableRowUtils::InsertNewRowBelowActionTooltip = LOCTEXT("FDataTableRowUtils_InsertNewRowBelowTooltip", "Insert a new Row below the current selection");
-
-const FText FDataTableRowUtils::MoveToBottomActionName = LOCTEXT("FDataTableRowUtils_MoveToBottom", "Move Row To Bottom");
-const FText FDataTableRowUtils::MoveToBottomActionTooltip = LOCTEXT("FDataTableRowUtils_MoveToBottomTooltip", "Move selected Row to the bottom");
-
-const FText FDataTableRowUtils::MoveToTopActionName = LOCTEXT("FDataTableRowUtils_MoveToTopAction", "Move Row to Top");
-const FText FDataTableRowUtils::MoveToTopActionTooltip = LOCTEXT("FDataTableRowUtils_MoveToTopActionTooltip", "Move selected Row to the top");
-
-
-TSharedRef<SWidget> FDataTableRowUtils::MakeRowActionsMenu(TSharedPtr<IDataTableEditor> Editor, FExecuteAction SearchForReferencesAction, FExecuteAction InsertNewRowAction,
-	FExecuteAction InsertNewRowAboveAction, FExecuteAction InsertNewRowBelowAction, FExecuteAction MoveToBottomAction, FExecuteAction MoveToTopAction)
-{
-	if (SearchForReferencesAction.IsBound() && InsertNewRowAction.IsBound() && InsertNewRowAboveAction.IsBound() 
-		&& InsertNewRowBelowAction.IsBound() )
-	{
-		FMenuBuilder MenuBuilder(true, Editor->GetToolkitCommands());
-		MenuBuilder.AddMenuEntry(InsertNewRowActionName, InsertNewRowActionTooltip,
-			FSlateIcon(FEditorStyle::GetStyleSetName(), "DataTableEditor.Add"), FUIAction(InsertNewRowAction));
-		MenuBuilder.AddMenuEntry(InsertNewRowAboveActionName, InsertNewRowAboveActionTooltip,
-			FSlateIcon(), FUIAction(InsertNewRowAboveAction));
-		MenuBuilder.AddMenuEntry(InsertNewRowBelowActionName, InsertNewRowBelowActionTooltip,
-			FSlateIcon(), FUIAction(InsertNewRowBelowAction));
-
-		MenuBuilder.AddMenuEntry(FGenericCommands::Get().Copy);
-		MenuBuilder.AddMenuEntry(FGenericCommands::Get().Paste);
-		MenuBuilder.AddMenuEntry(FGenericCommands::Get().Duplicate);
-		MenuBuilder.AddMenuEntry(FGenericCommands::Get().Rename);
-		MenuBuilder.AddMenuEntry(FGenericCommands::Get().Delete);
-		MenuBuilder.AddMenuSeparator();
-
-		MenuBuilder.AddMenuEntry(MoveToTopActionName, MoveToTopActionTooltip,
-			FSlateIcon(FEditorStyle::GetStyleSetName(), "Symbols.DoubleUpArrow"), FUIAction(MoveToTopAction));
-		MenuBuilder.AddMenuEntry(MoveToBottomActionName, MoveToBottomActionTooltip,
-			FSlateIcon(FEditorStyle::GetStyleSetName(), "Symbols.DoubleDownArrow"), FUIAction(MoveToBottomAction));
-		MenuBuilder.AddMenuSeparator();
-		
-		MenuBuilder.AddMenuEntry(SearchForReferencesActionName, SearchForReferencesActionTooltip, 
-			FSlateIcon(), FUIAction(SearchForReferencesAction));
-		return MenuBuilder.MakeWidget();
-	}
-
-	return SNullWidget::NullWidget;
-}
-
 void FDataTableRowUtils::AddSearchForReferencesContextMenu(FDetailWidgetRow& RowNameDetailWidget, FExecuteAction SearchForReferencesAction)
 {
 	if (SearchForReferencesAction.IsBound() && FEditorDelegates::OnOpenReferenceViewer.IsBound())
 	{
-		RowNameDetailWidget.AddCustomContextMenuAction(FUIAction(SearchForReferencesAction), SearchForReferencesActionName, 
-			SearchForReferencesActionTooltip, FSlateIcon());
+		RowNameDetailWidget.AddCustomContextMenuAction(FUIAction(SearchForReferencesAction), 
+			LOCTEXT("FDataTableRowUtils_SearchForReferences", "Find Row References"),
+			LOCTEXT("FDataTableRowUtils_SearchForReferencesTooltip", "Find assets that reference this Row"),
+			FSlateIcon());
 	}
 }
 
diff --git a/Engine/Source/Editor/DataTableEditor/Public/DataTableRowUtlis.h b/Engine/Source/Editor/DataTableEditor/Public/DataTableRowUtlis.h
--- a/Engine/Source/Editor/DataTableEditor/Public/DataTableRowUtlis.h
+++ b/Engine/Source/Editor/DataTableEditor/Public/DataTableRowUtlis.h
@@ -12,29 +12,9 @@ class FDetailWidgetRow;
 class DATATABLEEDITOR_API FDataTableRowUtils
 {
 public:
-	static TSharedRef<SWidget> MakeRowActionsMenu(TSharedPtr<class IDataTableEditor> Editor, FExecuteAction SearchForReferencesAction, FExecuteAction InsertNewRowAction,
-													FExecuteAction InsertNewRowAboveAction, FExecuteAction InsertNewRowBelowAction,
-												    FExecuteAction MoveToBottomAction, FExecuteAction MoveToTopAction);
 	static void AddSearchForReferencesContextMenu(FDetailWidgetRow& RowNameDetailWidget, FExecuteAction SearchForReferencesAction);
 
 private:
-	static const FText SearchForReferencesActionName;
-	static const FText SearchForReferencesActionTooltip;
-
-	static const FText InsertNewRowActionName;
-	static const FText InsertNewRowActionTooltip;
-
-	static const FText InsertNewRowAboveActionName;
-	static const FText InsertNewRowAboveActionTooltip;
-
-	static const FText InsertNewRowBelowActionName;
-	static const FText InsertNewRowBelowActionTooltip;
-
-	static const FText MoveToBottomActionName;
-	static const FText MoveToBottomActionTooltip;
-
-	static const FText MoveToTopActionName;
-	static const FText MoveToTopActionTooltip;
 
 };
 
diff --git a/Engine/Source/Editor/PropertyEditor/Private/SStructureDetailsView.cpp b/Engine/Source/Editor/PropertyEditor/Private/SStructureDetailsView.cpp
--- a/Engine/Source/Editor/PropertyEditor/Private/SStructureDetailsView.cpp
+++ b/Engine/Source/Editor/PropertyEditor/Private/SStructureDetailsView.cpp
@@ -254,6 +254,11 @@ void SStructureDetailsView::SetStructureData(TSharedPtr<FStructOnScope> InStruct
 	UpdateFilteredDetails();
 }
 
+void SStructureDetailsView::SetCustomName(const FText& Text)
+{
+	CustomName = Text;
+}
+
 void SStructureDetailsView::ForceRefresh()
 {
 	SetStructureData(StructData);
diff --git a/Engine/Source/Editor/PropertyEditor/Private/SStructureDetailsView.h b/Engine/Source/Editor/PropertyEditor/Private/SStructureDetailsView.h
--- a/Engine/Source/Editor/PropertyEditor/Private/SStructureDetailsView.h
+++ b/Engine/Source/Editor/PropertyEditor/Private/SStructureDetailsView.h
@@ -100,6 +100,9 @@ public:
 
 	TSharedPtr<class FComplexPropertyNode> GetRootNode();
 	const TSharedPtr<class FComplexPropertyNode> GetRootNode() const;
+
+	virtual void SetCustomName(const FText& Text) override;
+
 protected:
 
 	virtual void CustomUpdatePropertyMap(TSharedPtr<FDetailLayoutBuilderImpl>& InDetailLayout) override;
diff --git a/Engine/Source/Editor/PropertyEditor/Public/IStructureDetailsView.h b/Engine/Source/Editor/PropertyEditor/Public/IStructureDetailsView.h
--- a/Engine/Source/Editor/PropertyEditor/Public/IStructureDetailsView.h
+++ b/Engine/Source/Editor/PropertyEditor/Public/IStructureDetailsView.h
@@ -37,6 +37,13 @@ public:
 	 */
 	virtual void SetStructureData(TSharedPtr<class FStructOnScope> StructData) = 0;
 
+	/**
+	 * Set the CustomName used for display purposes.
+	 *
+	 * @param Text The FText to use for setting the Name
+	 */
+	virtual void SetCustomName(const FText& Text) = 0;
+
 public:
 
 	/**
