When forced sorting back to numerical sorting happens in Data Tables, the sort mode switches to the appropriate mode and direction.

diff --git a/Engine/Source/Editor/DataTableEditor/Private/DataTableEditor.cpp b/Engine/Source/Editor/DataTableEditor/Private/DataTableEditor.cpp
--- a/Engine/Source/Editor/DataTableEditor/Private/DataTableEditor.cpp
+++ b/Engine/Source/Editor/DataTableEditor/Private/DataTableEditor.cpp
@@ -315,6 +315,8 @@ void FDataTableEditor::OnAddClicked()
 
 		FDataTableEditorUtils::AddRow(Table, NewName);
 		FDataTableEditorUtils::SelectRow(Table, NewName);
+
+		SetDefaultSort();
 	}
 }
 
@@ -340,6 +342,7 @@ void FDataTableEditor::OnRemoveClicked()
 			}
 		}
 
+		SetDefaultSort();
 	}
 }
 
@@ -363,6 +366,8 @@ FReply FDataTableEditor::OnMoveToExtentClicked(FDataTableEditorUtils::ERowMoveDi
 		// We move by the row map size, as FDataTableEditorUtils::MoveRow will automatically clamp this as appropriate
 		FDataTableEditorUtils::MoveRow(Table, HighlightedRowName, MoveDirection, Table->GetRowMap().Num());
 		FDataTableEditorUtils::SelectRow(Table, HighlightedRowName);
+
+		SetDefaultSort();
 	}
 	return FReply::Handled();
 }
@@ -394,6 +399,12 @@ void FDataTableEditor::OnDuplicateClicked()
 	}
 }
 
+void FDataTableEditor::SetDefaultSort()
+{
+	SortMode = EColumnSortMode::Ascending;
+	SortByColumn = FDataTableEditor::RowNumberColumnId;
+}
+
 EColumnSortMode::Type FDataTableEditor::GetColumnSortMode(const FName ColumnId) const
 {
 	if (SortByColumn != ColumnId)
diff --git a/Engine/Source/Editor/DataTableEditor/Private/DataTableEditor.h b/Engine/Source/Editor/DataTableEditor/Private/DataTableEditor.h
--- a/Engine/Source/Editor/DataTableEditor/Private/DataTableEditor.h
+++ b/Engine/Source/Editor/DataTableEditor/Private/DataTableEditor.h
@@ -172,6 +172,7 @@ protected:
 	void OnPasteClicked();
 	void OnDuplicateClicked();
 
+	void SetDefaultSort();
 	EColumnSortMode::Type GetColumnSortMode(const FName ColumnId) const;
 	void OnColumnSortModeChanged(const EColumnSortPriority::Type SortPriority, const FName& ColumnId, const EColumnSortMode::Type InSortMode);
 	void OnColumnNumberSortModeChanged(const EColumnSortPriority::Type SortPriority, const FName& ColumnId, const EColumnSortMode::Type InSortMode);
diff --git a/Engine/Source/Editor/DataTableEditor/Private/SDataTableListViewRow.cpp b/Engine/Source/Editor/DataTableEditor/Private/SDataTableListViewRow.cpp
--- a/Engine/Source/Editor/DataTableEditor/Private/SDataTableListViewRow.cpp
+++ b/Engine/Source/Editor/DataTableEditor/Private/SDataTableListViewRow.cpp
@@ -89,6 +89,8 @@ void SDataTableListViewRow::OnInsertNewRow(ERowInsertionPosition InsertPosition)
 					FDataTableEditorUtils::AddRowAboveOrBelowSelection(SourceDataTable, *CurrentName, NewName, InsertPosition);
 				}
 				FDataTableEditorUtils::SelectRow(SourceDataTable, NewName);
+
+				DataTableEditorPtr->SetDefaultSort();
 			}
 		}
 	}
@@ -159,46 +161,52 @@ void SDataTableListViewRow::OnRowRenamed(const FText& Text, ETextCommit::Type Co
 
 	if (!GetCurrentNameAsText().EqualTo(Text) && DataTable)
 	{
-		const TArray<FName> RowNames = DataTable->GetRowNames();
-
-		if (Text.IsEmptyOrWhitespace() || !FName::IsValidXName(Text.ToString(), INVALID_NAME_CHARACTERS))
+		if (FDataTableEditor* DataTableEditorPtr = DataTableEditor.Pin().Get())
 		{
-			// Only pop up the error dialog if the rename was caused by the user's action
-			if ((CommitType == ETextCommit::OnEnter) || (CommitType == ETextCommit::OnUserMovedFocus))
+
+			const TArray<FName> RowNames = DataTable->GetRowNames();
+
+			if (Text.IsEmptyOrWhitespace() || !FName::IsValidXName(Text.ToString(), INVALID_NAME_CHARACTERS))
 			{
-				// popup an error dialog here
-				const FText Message = FText::Format(LOCTEXT("InvalidRowName", "'{0}' is not a valid row name"), Text);
-				FMessageDialog::Open(EAppMsgType::Ok, Message);
+				// Only pop up the error dialog if the rename was caused by the user's action
+				if ((CommitType == ETextCommit::OnEnter) || (CommitType == ETextCommit::OnUserMovedFocus))
+				{
+					// popup an error dialog here
+					const FText Message = FText::Format(LOCTEXT("InvalidRowName", "'{0}' is not a valid row name"), Text);
+					FMessageDialog::Open(EAppMsgType::Ok, Message);
+				}
+				return;
 			}
-			return;
-		}
-		const FName NewName = DataTableUtils::MakeValidName(Text.ToString());
-		if (NewName == NAME_None)
-		{
-			// popup an error dialog here
-			const FText Message = FText::Format(LOCTEXT("InvalidRowName", "'{0}' is not a valid row name"), Text);
-			FMessageDialog::Open(EAppMsgType::Ok, Message);
-
-			return;
-		}
-		for (const FName& Name : RowNames)
-		{
-			if (Name.IsValid() && (Name == NewName))
+			const FName NewName = DataTableUtils::MakeValidName(Text.ToString());
+			if (NewName == NAME_None)
 			{
-				//the name already exists
 				// popup an error dialog here
-				const FText Message = FText::Format(LOCTEXT("DuplicateRowName", "'{0}' is already used as a row name in this table"), Text);
+				const FText Message = FText::Format(LOCTEXT("InvalidRowName", "'{0}' is not a valid row name"), Text);
 				FMessageDialog::Open(EAppMsgType::Ok, Message);
+
 				return;
+			}
+			for (const FName& Name : RowNames)
+			{
+				if (Name.IsValid() && (Name == NewName))
+				{
+					//the name already exists
+					// popup an error dialog here
+					const FText Message = FText::Format(LOCTEXT("DuplicateRowName", "'{0}' is already used as a row name in this table"), Text);
+					FMessageDialog::Open(EAppMsgType::Ok, Message);
+					return;
 
+				}
 			}
-		}
 
-		const FName OldName = GetCurrentName();
-		FDataTableEditorUtils::RenameRow(DataTable, OldName, NewName);
-		FDataTableEditorUtils::SelectRow(DataTable, NewName);
+			const FName OldName = GetCurrentName();
+			FDataTableEditorUtils::RenameRow(DataTable, OldName, NewName);
+			FDataTableEditorUtils::SelectRow(DataTable, NewName);
+
+			DataTableEditorPtr->SetDefaultSort();
 
-		*CurrentName = NewName;
+			*CurrentName = NewName;
+		}
 	}
 }
 
