Bug fix: Removing the last Data Table row operation did not properly update the Data Table window, fixed. In addition, the 2 different delete operations (delete by clicking the "Remove" button and delete by pressing the "Del" key) for that table functionality are now unified into the same code.

diff --git a/Engine/Source/Editor/DataTableEditor/Private/DataTableEditor.cpp b/Engine/Source/Editor/DataTableEditor/Private/DataTableEditor.cpp
--- a/Engine/Source/Editor/DataTableEditor/Private/DataTableEditor.cpp
+++ b/Engine/Source/Editor/DataTableEditor/Private/DataTableEditor.cpp
@@ -322,28 +322,7 @@ void FDataTableEditor::OnAddClicked()
 
 void FDataTableEditor::OnRemoveClicked()
 {
-	UDataTable* Table = GetEditableDataTable();
-
-	if (Table)
-	{
-
-		const int32 RowToRemoveIndex = VisibleRows.IndexOfByPredicate([&](const FDataTableEditorRowListViewDataPtr& InRowName) -> bool
-		{
-			return InRowName->RowId == HighlightedRowName;
-		});
-
-		if (FDataTableEditorUtils::RemoveRow(Table, HighlightedRowName))
-		{
-			// Try and keep the same row index selected
-			const int32 RowIndexToSelect = FMath::Clamp(RowToRemoveIndex, 0, VisibleRows.Num() - 1);
-			if (VisibleRows.IsValidIndex(RowIndexToSelect))
-			{
-				FDataTableEditorUtils::SelectRow(Table, VisibleRows[RowIndexToSelect]->RowId);
-			}
-		}
-
-		SetDefaultSort();
-	}
+	DeleteSelectedRow();
 }
 
 FReply FDataTableEditor::OnMoveRowClicked(FDataTableEditorUtils::ERowMoveDirection MoveDirection)
@@ -909,14 +888,30 @@ void FDataTableEditor::RenameSelectedRowCommand()
 
 void FDataTableEditor::DeleteSelectedRow()
 {
-	UDataTable* TablePtr = Cast<UDataTable>(GetEditingObject());
-
-	if (HighlightedRowName == NAME_None || TablePtr == nullptr)
+	if (UDataTable* Table = GetEditableDataTable())
 	{
-		return;
+		// We must perform this before removing the row
+		const int32 RowToRemoveIndex = VisibleRows.IndexOfByPredicate([&](const FDataTableEditorRowListViewDataPtr& InRowName) -> bool
+		{
+			return InRowName->RowId == HighlightedRowName;
+		});
+		// Remove row
+		if (FDataTableEditorUtils::RemoveRow(Table, HighlightedRowName))
+		{
+			// Try and keep the same row index selected
+			const int32 RowIndexToSelect = FMath::Clamp(RowToRemoveIndex, 0, VisibleRows.Num() - 1);
+			if (VisibleRows.IsValidIndex(RowIndexToSelect))
+			{
+				FDataTableEditorUtils::SelectRow(Table, VisibleRows[RowIndexToSelect]->RowId);
+			}
+			// Refresh list. Otherwise, the removed row would still appear in the screen until the next list refresh. An
+			// analog of CellsListView->RequestListRefresh() also occurs inside FDataTableEditorUtils::SelectRow
+			else
+			{
+				CellsListView->RequestListRefresh();
+			}
+		}
 	}
-
-	FDataTableEditorUtils::RemoveRow(TablePtr, HighlightedRowName);
 }
 
 FText FDataTableEditor::GetFilterText() const
