UE-62702 - Clicking the button to remove a Data Table row while no row is selected causes the first row in the table to be selected

The various commands that depend on a valid selection (eg: Copy, Duplicate, etc) need to use a FCanExecuteAction delegate to ensure that a row is currently highlighted. For now, we can just use CanEditTable for that check.

The Jira only calls out the toolbar as having this issue, but we need to add FCanExecuteAction to the commands as well. This ensures that these commands are disabled if they're ever added to a menu in the future.

diff --git a/Engine/Source/Editor/DataTableEditor/Private/DataTableEditor.cpp b/Engine/Source/Editor/DataTableEditor/Private/DataTableEditor.cpp
--- a/Engine/Source/Editor/DataTableEditor/Private/DataTableEditor.cpp
+++ b/Engine/Source/Editor/DataTableEditor/Private/DataTableEditor.cpp
@@ -269,11 +269,11 @@ void FDataTableEditor::InitDataTableEditor( const EToolkitMode::Type Mode, const
 	}*/
 
 	// asset editor commands here
-	ToolkitCommands->MapAction(FGenericCommands::Get().Copy, FExecuteAction::CreateSP(this, &FDataTableEditor::CopySelectedRow));
-	ToolkitCommands->MapAction(FGenericCommands::Get().Paste, FExecuteAction::CreateSP(this, &FDataTableEditor::PasteOnSelectedRow));
-	ToolkitCommands->MapAction(FGenericCommands::Get().Duplicate, FExecuteAction::CreateSP(this, &FDataTableEditor::DuplicateSelectedRow));
-	ToolkitCommands->MapAction(FGenericCommands::Get().Rename, FExecuteAction::CreateSP(this, &FDataTableEditor::RenameSelectedRowCommand));
-	ToolkitCommands->MapAction(FGenericCommands::Get().Delete, FExecuteAction::CreateSP(this, &FDataTableEditor::DeleteSelectedRow));
+	ToolkitCommands->MapAction(FGenericCommands::Get().Copy, FExecuteAction::CreateSP(this, &FDataTableEditor::CopySelectedRow), FCanExecuteAction::CreateSP(this, &FDataTableEditor::CanEditTable));
+	ToolkitCommands->MapAction(FGenericCommands::Get().Paste, FExecuteAction::CreateSP(this, &FDataTableEditor::PasteOnSelectedRow), FCanExecuteAction::CreateSP(this, &FDataTableEditor::CanEditTable));
+	ToolkitCommands->MapAction(FGenericCommands::Get().Duplicate, FExecuteAction::CreateSP(this, &FDataTableEditor::DuplicateSelectedRow), FCanExecuteAction::CreateSP(this, &FDataTableEditor::CanEditTable));
+	ToolkitCommands->MapAction(FGenericCommands::Get().Rename, FExecuteAction::CreateSP(this, &FDataTableEditor::RenameSelectedRowCommand), FCanExecuteAction::CreateSP(this, &FDataTableEditor::CanEditTable));
+	ToolkitCommands->MapAction(FGenericCommands::Get().Delete, FExecuteAction::CreateSP(this, &FDataTableEditor::DeleteSelectedRow), FCanExecuteAction::CreateSP(this, &FDataTableEditor::CanEditTable));
 }
 
 bool FDataTableEditor::CanEditRows() const
@@ -368,6 +368,11 @@ void FDataTableEditor::OnDuplicateClicked()
 	}
 }
 
+bool FDataTableEditor::CanEditTable() const
+{
+	return HighlightedRowName != NAME_None;
+}
+
 void FDataTableEditor::SetDefaultSort()
 {
 	SortMode = EColumnSortMode::Ascending;
@@ -533,25 +538,33 @@ void FDataTableEditor::FillToolbar(FToolBarBuilder& ToolbarBuilder)
 			LOCTEXT("AddRowToolTip", "Add a new row to the Data Table"),
 			FSlateIcon(FEditorStyle::GetStyleSetName(), "DataTableEditor.Add"));
 		ToolbarBuilder.AddToolBarButton(
-			FUIAction(FExecuteAction::CreateSP(this, &FDataTableEditor::OnCopyClicked)),
+			FUIAction(
+				FExecuteAction::CreateSP(this, &FDataTableEditor::OnCopyClicked),
+				FCanExecuteAction::CreateSP(this, &FDataTableEditor::CanEditTable)),
 			NAME_None,
 			LOCTEXT("CopyIconText", "Copy"),
 			LOCTEXT("CopyToolTip", "Copy the currently selected row"),
 			FSlateIcon(FEditorStyle::GetStyleSetName(), "DataTableEditor.Copy"));
 		ToolbarBuilder.AddToolBarButton(
-			FUIAction(FExecuteAction::CreateSP(this, &FDataTableEditor::OnPasteClicked)),
+			FUIAction(
+				FExecuteAction::CreateSP(this, &FDataTableEditor::OnPasteClicked),
+				FCanExecuteAction::CreateSP(this, &FDataTableEditor::CanEditTable)),
 			NAME_None,
 			LOCTEXT("PasteIconText", "Paste"),
 			LOCTEXT("PasteToolTip", "Paste on the currently selected row"),
 			FSlateIcon(FEditorStyle::GetStyleSetName(), "DataTableEditor.Paste"));
 		ToolbarBuilder.AddToolBarButton(
-			FUIAction(FExecuteAction::CreateSP(this, &FDataTableEditor::OnDuplicateClicked)),
+			FUIAction(
+				FExecuteAction::CreateSP(this, &FDataTableEditor::OnDuplicateClicked),
+				FCanExecuteAction::CreateSP(this, &FDataTableEditor::CanEditTable)),
 			NAME_None,
 			LOCTEXT("DuplicateIconText", "Duplicate"),
 			LOCTEXT("DuplicateToolTip", "Duplicate the currently selected row"),
 			FSlateIcon(FEditorStyle::GetStyleSetName(), "DataTableEditor.Duplicate"));
 		ToolbarBuilder.AddToolBarButton(
-			FUIAction(FExecuteAction::CreateSP(this, &FDataTableEditor::OnRemoveClicked)),
+			FUIAction(
+				FExecuteAction::CreateSP(this, &FDataTableEditor::OnRemoveClicked),
+				FCanExecuteAction::CreateSP(this, &FDataTableEditor::CanEditTable)),
 			NAME_None,
 			LOCTEXT("RemoveRowIconText", "Remove"),
 			LOCTEXT("RemoveRowToolTip", "Remove the currently selected row from the Data Table"),
diff --git a/Engine/Source/Editor/DataTableEditor/Private/DataTableEditor.h b/Engine/Source/Editor/DataTableEditor/Private/DataTableEditor.h
--- a/Engine/Source/Editor/DataTableEditor/Private/DataTableEditor.h
+++ b/Engine/Source/Editor/DataTableEditor/Private/DataTableEditor.h
@@ -164,6 +164,7 @@ protected:
 	void OnCopyClicked();
 	void OnPasteClicked();
 	void OnDuplicateClicked();
+	bool CanEditTable() const;
 
 	void SetDefaultSort();
 	EColumnSortMode::Type GetColumnSortMode(const FName ColumnId) const;
