fix EditionCondition error when item is not selected in Datatable

diff --git a/Engine/Source/Editor/DataTableEditor/Private/SRowEditor.cpp b/Engine/Source/Editor/DataTableEditor/Private/SRowEditor.cpp
--- a/Engine/Source/Editor/DataTableEditor/Private/SRowEditor.cpp
+++ b/Engine/Source/Editor/DataTableEditor/Private/SRowEditor.cpp
@@ -231,7 +231,10 @@ void SRowEditor::Restore()
 	if (StructureDetailsView.IsValid())
 	{
 		StructureDetailsView->SetCustomName(FText::FromName(FinalName));
-		StructureDetailsView->SetStructureData(CurrentRow);
+		if (!FinalName.IsNone())
+		{
+			StructureDetailsView->SetStructureData(CurrentRow);
+		}
 	}
 
 	RowSelectedCallback.ExecuteIfBound(FinalName);
