Enabling search in data table row editor.

#Synced-CL 49072585

diff --git a/Engine/Source/Editor/DataTableEditor/Private/SRowEditor.cpp b/Engine/Source/Editor/DataTableEditor/Private/SRowEditor.cpp
--- a/Engine/Source/Editor/DataTableEditor/Private/SRowEditor.cpp
+++ b/Engine/Source/Editor/DataTableEditor/Private/SRowEditor.cpp
@@ -441,8 +441,9 @@ void SRowEditor::ConstructInternal(UDataTable* Changed)
 	{
 		FPropertyEditorModule& PropertyModule = FModuleManager::LoadModuleChecked<FPropertyEditorModule>("PropertyEditor");
 		FDetailsViewArgs ViewArgs;
-		ViewArgs.bAllowSearch = false;
+		ViewArgs.bAllowSearch = true;
 		ViewArgs.bHideSelectionTip = false;
+		ViewArgs.bShowOptions = false;
 		ViewArgs.bShowObjectLabel = false;
 		ViewArgs.NotifyHook = this;
 
