Upgraded DataTableEditor row context menu to use UToolMenus

Replaces the FMenuBuilder-based row context menu in SDataTableListViewRow with a named UToolMenus menu (DataTableEditor.RowContextMenu), allowing other modules to extend it with their own entries. Row-specific data is passed via a new UDataTableRowMenuContext object added to the FToolMenuContext each time the menu is opened. Data table asset can be accessed via the UDatatableRowMenuContext as it extends from UAssetEditorToolkitMenuContext by looking to the edited objects.

Also fixes a bug where the entire context menu was silently suppressed if FEditorDelegates::OnOpenReferenceViewer had no listeners, the delegate check is now scoped to just the "Find Row References" entry as a CanExecute condition.

#rnx
#Synced-CL 52307970

diff --git a/Engine/Source/Editor/DataTableEditor/DataTableEditor.Build.cs b/Engine/Source/Editor/DataTableEditor/DataTableEditor.Build.cs
--- a/Engine/Source/Editor/DataTableEditor/DataTableEditor.Build.cs
+++ b/Engine/Source/Editor/DataTableEditor/DataTableEditor.Build.cs
@@ -22,7 +22,8 @@ public class DataTableEditor : ModuleRules
 				"EditorFramework",
 				"UnrealEd",
 				"Json",
-				"DeveloperSettings"
+				"DeveloperSettings",
+				"ToolMenus"
 			}
 			);
 
diff --git a/Engine/Source/Editor/DataTableEditor/Private/SDataTableListViewRow.cpp b/Engine/Source/Editor/DataTableEditor/Private/SDataTableListViewRow.cpp
--- a/Engine/Source/Editor/DataTableEditor/Private/SDataTableListViewRow.cpp
+++ b/Engine/Source/Editor/DataTableEditor/Private/SDataTableListViewRow.cpp
@@ -2,6 +2,9 @@
 
 #include "SDataTableListViewRow.h"
 
+#include "DataTableRowMenuContext.h"
+#include "ToolMenu.h"
+#include "ToolMenus.h"
 #include "AssetRegistry/AssetData.h"
 #include "Containers/Map.h"
 #include "DataTableEditor.h"
@@ -69,7 +72,7 @@ void SDataTableListViewRow::Construct(const FArguments& InArgs, const TSharedRef
 
 FReply SDataTableListViewRow::OnMouseButtonUp(const FGeometry& MyGeometry, const FPointerEvent& MouseEvent)
 {
-	if (IsEditable && MouseEvent.GetEffectingButton() == EKeys::RightMouseButton && RowDataPtr.IsValid() && FEditorDelegates::OnOpenReferenceViewer.IsBound() && DataTableEditor.IsValid())
+	if (IsEditable && MouseEvent.GetEffectingButton() == EKeys::RightMouseButton && RowDataPtr.IsValid() && DataTableEditor.IsValid())
 	{
 		FDataTableEditorUtils::SelectRow(DataTableEditor.Pin()->GetDataTable(), RowDataPtr->RowId);
 
@@ -83,53 +86,6 @@ FReply SDataTableListViewRow::OnMouseButtonUp(const FGeometry& MyGeometry, const
 	return STableRow::OnMouseButtonUp(MyGeometry, MouseEvent);
 }
 
-void SDataTableListViewRow::OnSearchForReferences()
-{
-	if (DataTableEditor.IsValid() && RowDataPtr.IsValid())
-	{
-		if (FDataTableEditor* DataTableEditorPtr = DataTableEditor.Pin().Get())
-		{
-			UDataTable* SourceDataTable = const_cast<UDataTable*>(DataTableEditorPtr->GetDataTable());
-
-			TArray<FAssetIdentifier> AssetIdentifiers;
-			AssetIdentifiers.Add(FAssetIdentifier(SourceDataTable, RowDataPtr->RowId));
-
-			FEditorDelegates::OnOpenReferenceViewer.Broadcast(AssetIdentifiers, FReferenceViewerParams());
-		}
-	}
-}
-
-void SDataTableListViewRow::OnInsertNewRow(ERowInsertionPosition InsertPosition)
-{
-	if (DataTableEditor.IsValid() && RowDataPtr.IsValid())
-	{
-		if (FDataTableEditor* DataTableEditorPtr = DataTableEditor.Pin().Get())
-		{
-			UDataTable* SourceDataTable = const_cast<UDataTable*>(DataTableEditorPtr->GetDataTable());
-
-			if (SourceDataTable)
-			{
-				FName NewName = DataTableUtils::MakeValidName(TEXT("NewRow"));
-				while (SourceDataTable->GetRowMap().Contains(NewName))
-				{
-					NewName.SetNumber(NewName.GetNumber() + 1);
-				}
-
-				if (InsertPosition == ERowInsertionPosition::Bottom)
-				{
-					FDataTableEditorUtils::AddRow(SourceDataTable, NewName);
-				}
-				else
-				{
-					FDataTableEditorUtils::AddRowAboveOrBelowSelection(SourceDataTable, *CurrentName, NewName, InsertPosition);
-				}
-				FDataTableEditorUtils::SelectRow(SourceDataTable, NewName);
-
-				DataTableEditorPtr->SetDefaultSort();
-			}
-		}
-	}
-}
 
 FReply SDataTableListViewRow::OnRowDrop(const FDragDropEvent& DragDropEvent)
 {
@@ -443,73 +399,164 @@ const FSlateBrush* SDataTableListViewRow::GetBorder() const
 	}
 }
 
-void SDataTableListViewRow::OnMoveToExtentClicked(FDataTableEditorUtils::ERowMoveDirection MoveDirection)
-{
-	if (DataTableEditor.IsValid())
-	{
-		TSharedPtr<FDataTableEditor> DataTableEditorPtr = DataTableEditor.Pin();
-		DataTableEditorPtr->OnMoveToExtentClicked(MoveDirection);
-	}
-}
-
 
 TSharedRef<SWidget> SDataTableListViewRow::MakeRowActionsMenu()
 {
-	FMenuBuilder MenuBuilder(true, DataTableEditor.Pin()->GetToolkitCommands());
+	static const FName MenuName = "DataTableEditor.RowContextMenu";
 
-	MenuBuilder.AddMenuEntry(
-		LOCTEXT("DataTableRowMenuActions_InsertNewRow", "Insert New Row"),
-		LOCTEXT("DataTableRowMenuActions_InsertNewRowTooltip", "Insert a new row"),
-		FSlateIcon(FAppStyle::GetAppStyleSetName(), "DataTableEditor.Add"), 
-		FUIAction(FExecuteAction::CreateSP(this, &SDataTableListViewRow::OnInsertNewRow, ERowInsertionPosition::Bottom))
-	);
-	
-	MenuBuilder.AddMenuEntry(
-		LOCTEXT("DataTableRowMenuActions_InsertNewRowAbove", "Insert New Row Above"),
-		LOCTEXT("DataTableRowMenuActions_InsertNewRowAboveTooltip", "Insert a new Row above the current selection"),
-		FSlateIcon(), 
-		FUIAction(FExecuteAction::CreateSP(this, &SDataTableListViewRow::OnInsertNewRow, ERowInsertionPosition::Above))
-	);
-	
-	MenuBuilder.AddMenuEntry(
-		LOCTEXT("DataTableRowMenuActions_InsertNewRowBelow", "Insert New Row Below"),
-		LOCTEXT("DataTableRowMenuActions_InsertNewRowBelowTooltip", "Insert a new Row below the current selection"),
-		FSlateIcon(), 
-		FUIAction(FExecuteAction::CreateSP(this, &SDataTableListViewRow::OnInsertNewRow, ERowInsertionPosition::Below))
-	);
+	if (!UToolMenus::Get()->IsMenuRegistered(MenuName))
+	{
+		UToolMenu* Menu = UToolMenus::Get()->RegisterMenu(MenuName);
+		Menu->AddDynamicSection("DefaultRowActions", FNewToolMenuDelegate::CreateLambda([](UToolMenu* InMenu)
+		{
+			// Validate context.
+			UDataTableRowMenuContext* Ctx = InMenu->FindContext<UDataTableRowMenuContext>();
+			if (!Ctx || !Ctx->Toolkit.IsValid() || !Ctx->RowDataPtr)
+			{
+				return;
+			}
 
-	MenuBuilder.AddMenuEntry(FGenericCommands::Get().Copy);
-	MenuBuilder.AddMenuEntry(FGenericCommands::Get().Paste);
-	MenuBuilder.AddMenuEntry(FGenericCommands::Get().Duplicate);
-	MenuBuilder.AddMenuEntry(FGenericCommands::Get().Rename);
-	MenuBuilder.AddMenuEntry(FGenericCommands::Get().Delete);
+			// Get data table editor.
+			TWeakPtr<FDataTableEditor> WeakEditor = StaticCastWeakPtr<FDataTableEditor>(Ctx->Toolkit);
+			const FName RowId = Ctx->RowDataPtr->RowId;
 
-	MenuBuilder.AddMenuSeparator();
+			// Insert actions
+			FToolMenuSection& InsertSection = InMenu->FindOrAddSection("InsertActions");
 
-	MenuBuilder.AddMenuEntry(
-		LOCTEXT("DataTableRowMenuActions_MoveToTopAction", "Move Row to Top"),
-		LOCTEXT("DataTableRowMenuActions_MoveToTopActionTooltip", "Move selected Row to the top"),
-		FSlateIcon(FAppStyle::GetAppStyleSetName(), "Symbols.DoubleUpArrow"), 
-		FUIAction(FExecuteAction::CreateSP(this, &SDataTableListViewRow::OnMoveToExtentClicked, FDataTableEditorUtils::ERowMoveDirection::Up))
-	);
-	
-	MenuBuilder.AddMenuEntry(
-		LOCTEXT("DataTableRowMenuActions_MoveToBottom", "Move Row To Bottom"),
-		LOCTEXT("DataTableRowMenuActions_MoveToBottomTooltip", "Move selected Row to the bottom"),
-		FSlateIcon(FAppStyle::GetAppStyleSetName(), "Symbols.DoubleDownArrow"), 
-		FUIAction(FExecuteAction::CreateSP(this, &SDataTableListViewRow::OnMoveToExtentClicked, FDataTableEditorUtils::ERowMoveDirection::Down))
-	);
-	
-	MenuBuilder.AddMenuSeparator();
-
-	MenuBuilder.AddMenuEntry(
-		NSLOCTEXT("FDataTableRowUtils", "FDataTableRowUtils_SearchForReferences", "Find Row References"),
-		NSLOCTEXT("FDataTableRowUtils", "FDataTableRowUtils_SearchForReferencesTooltip", "Find assets that reference this Row"),
-		FSlateIcon(), 
-		FUIAction(FExecuteAction::CreateSP(this, &SDataTableListViewRow::OnSearchForReferences))
-	);
+			auto MakeInsertAction = [WeakEditor, RowId](ERowInsertionPosition InsertPos) -> FUIAction
+			{
+				return FUIAction(FExecuteAction::CreateLambda([WeakEditor, RowId, InsertPos]()
+				{
+					TSharedPtr<FDataTableEditor> Editor = WeakEditor.Pin();
+					if (!Editor.IsValid())
+					{
+						return;
+					}
+						
+					UDataTable* Table = Cast<UDataTable>(Editor->GetEditingObject());
+					if (!Table)
+					{
+						return;
+					}
+
+					FName NewName = DataTableUtils::MakeValidName(TEXT("NewRow"));
+					while (Table->GetRowMap().Contains(NewName))
+					{
+						NewName.SetNumber(NewName.GetNumber() + 1);
+					}
+
+					if (InsertPos == ERowInsertionPosition::Bottom)
+					{
+						FDataTableEditorUtils::AddRow(Table, NewName);
+					}
+					else
+					{
+						FDataTableEditorUtils::AddRowAboveOrBelowSelection(Table, RowId, NewName, InsertPos);
+					}
+
+					FDataTableEditorUtils::SelectRow(Table, NewName);
+					Editor->SetDefaultSort();
+				}));
+			};
+
+			InsertSection.AddMenuEntry("InsertNewRow",
+				LOCTEXT("DataTableRowMenuActions_InsertNewRow", "Insert New Row"),
+				LOCTEXT("DataTableRowMenuActions_InsertNewRowTooltip", "Insert a new row"),
+				FSlateIcon(FAppStyle::GetAppStyleSetName(), "DataTableEditor.Add"),
+				MakeInsertAction(ERowInsertionPosition::Bottom)
+			);
+			InsertSection.AddMenuEntry("InsertNewRowAbove",
+				LOCTEXT("DataTableRowMenuActions_InsertNewRowAbove", "Insert New Row Above"),
+				LOCTEXT("DataTableRowMenuActions_InsertNewRowAboveTooltip", "Insert a new Row above the current selection"),
+				FSlateIcon(),
+				MakeInsertAction(ERowInsertionPosition::Above)
+			);
+			InsertSection.AddMenuEntry("InsertNewRowBelow",
+				LOCTEXT("DataTableRowMenuActions_InsertNewRowBelow", "Insert New Row Below"),
+				LOCTEXT("DataTableRowMenuActions_InsertNewRowBelowTooltip", "Insert a new Row below the current selection"),
+				FSlateIcon(),
+				MakeInsertAction(ERowInsertionPosition::Below)
+			);
+
+			// Edit actions — actions are resolved from the command list appended to the context
+			FToolMenuSection& EditSection = InMenu->FindOrAddSection("EditActions");
+			EditSection.AddEntry(FToolMenuEntry::InitMenuEntry(FGenericCommands::Get().Copy));
+			EditSection.AddEntry(FToolMenuEntry::InitMenuEntry(FGenericCommands::Get().Paste));
+			EditSection.AddEntry(FToolMenuEntry::InitMenuEntry(FGenericCommands::Get().Duplicate));
+			EditSection.AddEntry(FToolMenuEntry::InitMenuEntry(FGenericCommands::Get().Rename));
+			EditSection.AddEntry(FToolMenuEntry::InitMenuEntry(FGenericCommands::Get().Delete));
+				
+			// Move actions
+			FToolMenuSection& MoveSection = InMenu->FindOrAddSection("MoveActions");
+			MoveSection.AddSeparator("MoveActionsSeparator");
+			MoveSection.AddMenuEntry("MoveRowToTop",
+				LOCTEXT("DataTableRowMenuActions_MoveToTopAction", "Move Row to Top"),
+				LOCTEXT("DataTableRowMenuActions_MoveToTopActionTooltip", "Move selected Row to the top"),
+				FSlateIcon(FAppStyle::GetAppStyleSetName(), "Symbols.DoubleUpArrow"),
+				FUIAction(FExecuteAction::CreateLambda([WeakEditor]()
+				{
+					if (TSharedPtr<FDataTableEditor> Editor = WeakEditor.Pin())
+					{
+						Editor->OnMoveToExtentClicked(FDataTableEditorUtils::ERowMoveDirection::Up);
+					}
+				}))
+			);
+			MoveSection.AddMenuEntry("MoveRowToBottom",
+				LOCTEXT("DataTableRowMenuActions_MoveToBottom", "Move Row To Bottom"),
+				LOCTEXT("DataTableRowMenuActions_MoveToBottomTooltip", "Move selected Row to the bottom"),
+				FSlateIcon(FAppStyle::GetAppStyleSetName(), "Symbols.DoubleDownArrow"),
+				FUIAction(FExecuteAction::CreateLambda([WeakEditor]()
+				{
+					if (TSharedPtr<FDataTableEditor> Editor = WeakEditor.Pin())
+					{
+						Editor->OnMoveToExtentClicked(FDataTableEditorUtils::ERowMoveDirection::Down);
+					}
+				}))
+			);
+
+			// Reference actions
+			FToolMenuSection& RefSection = InMenu->FindOrAddSection("ReferenceActions");
+			RefSection.AddSeparator("ReferenceActionsSeparator");
+			RefSection.AddMenuEntry("FindRowReferences",
+				NSLOCTEXT("FDataTableRowUtils", "FDataTableRowUtils_SearchForReferences", "Find Row References"),
+				NSLOCTEXT("FDataTableRowUtils", "FDataTableRowUtils_SearchForReferencesTooltip", "Find assets that reference this Row"),
+				FSlateIcon(),
+				FUIAction(
+					FExecuteAction::CreateLambda([WeakEditor, RowId]()
+					{
+						TSharedPtr<FDataTableEditor> Editor = WeakEditor.Pin();
+						if (!Editor.IsValid())
+						{
+							return;
+						}
+							
+						UDataTable* Table = Cast<UDataTable>(Editor->GetEditingObject());
+						if (!Table)
+						{
+							return;
+						}
+							
+						TArray<FAssetIdentifier> AssetIdentifiers;
+						AssetIdentifiers.Add(FAssetIdentifier(Table, RowId));
+						FEditorDelegates::OnOpenReferenceViewer.Broadcast(AssetIdentifiers, FReferenceViewerParams());
+					}),
+					FCanExecuteAction::CreateLambda([]()
+					{
+						return FEditorDelegates::OnOpenReferenceViewer.IsBound();
+					})
+				)
+			);
+		}));
+	}
+
+	UDataTableRowMenuContext* RowContext = NewObject<UDataTableRowMenuContext>();
+	RowContext->Toolkit = DataTableEditor;
+	RowContext->RowDataPtr = RowDataPtr;
+
+	FToolMenuContext MenuContext(DataTableEditor.Pin()->GetToolkitCommands());
+	MenuContext.AddObject(RowContext);
 
-	return MenuBuilder.MakeWidget();
+	return UToolMenus::Get()->GenerateWidget(MenuName, MenuContext);
 }
 
 FDataTableRowDragDropOp::FDataTableRowDragDropOp(TSharedPtr<SDataTableListViewRow> InRow)
diff --git a/Engine/Source/Editor/DataTableEditor/Private/SDataTableListViewRow.h b/Engine/Source/Editor/DataTableEditor/Private/SDataTableListViewRow.h
--- a/Engine/Source/Editor/DataTableEditor/Private/SDataTableListViewRow.h
+++ b/Engine/Source/Editor/DataTableEditor/Private/SDataTableListViewRow.h
@@ -101,9 +101,6 @@ public:
 
 private:
 
-	void OnSearchForReferences();
-	void OnInsertNewRow(ERowInsertionPosition InsertPosition);
-	
 	FReply OnRowDrop(const FDragDropEvent& DragDropEvent);
 
 	TSharedRef<SWidget> MakeCellWidget(const int32 InRowIndex, const FName& InColumnId);
@@ -113,8 +110,6 @@ private:
 
 	virtual const FSlateBrush* GetBorder() const;
 
-	void OnMoveToExtentClicked(FDataTableEditorUtils::ERowMoveDirection MoveDirection);
-
 	TSharedRef<SWidget> MakeRowActionsMenu();
 
 	TSharedPtr<SInlineEditableTextBlock> InlineEditableText;
diff --git a/Engine/Source/Editor/DataTableEditor/Public/DataTableRowMenuContext.h b/Engine/Source/Editor/DataTableEditor/Public/DataTableRowMenuContext.h
new file mode 100644
--- /dev/null
+++ b/Engine/Source/Editor/DataTableEditor/Public/DataTableRowMenuContext.h
@@ -0,0 +1,19 @@
+// Copyright Epic Games, Inc. All Rights Reserved.
+
+#pragma once
+
+#include "DataTableEditorUtils.h"
+#include "Toolkits/AssetEditorToolkitMenuContext.h"
+#include "DataTableRowMenuContext.generated.h"
+
+class FDataTableEditor;
+
+UCLASS()
+class DATATABLEEDITOR_API UDataTableRowMenuContext : public UAssetEditorToolkitMenuContext
+{
+	GENERATED_BODY()
+	
+public:
+
+	FDataTableEditorRowListViewDataPtr RowDataPtr;
+};
